CaptureManager: Move plugin to Virtual Production directory

diff --git a/Engine/Plugins/VirtualProduction/CaptureManager/CaptureManagerApp/CaptureManagerApp.uplugin b/Engine/Plugins/VirtualProduction/CaptureManager/CaptureManagerApp/CaptureManagerApp.uplugin
new file mode 100644
--- /dev/null
+++ b/Engine/Plugins/VirtualProduction/CaptureManager/CaptureManagerApp/CaptureManagerApp.uplugin
@@ -0,0 +1,134 @@
+{
+	"FileVersion": 1,
+	"Version": 1,
+	"VersionName": "1.0.0",
+	"FriendlyName": "Capture Manager Application Plugin",
+	"Description": "The Capture Manager allows control and monitoring of the capture device, obtains and transcodes the data from the devices and upload the data for import to the UE",
+	"Category": "Virtual Production",
+	"CreatedBy": "Epic Games, Inc.",
+	"CreatedByURL": "http://epicgames.com",
+	"DocsURL": "",
+	"MarketplaceURL": "",
+	"SupportURL": "",
+	"CanContainContent": true,
+	"EnabledByDefault": false,
+	"IsBetaVersion": false,
+	"IsExperimentalVersion": true,
+	"Installed": false,
+	"Hidden" :  true,
+	"SupportedPrograms": [ "LiveLinkHub" ],
+	"Modules": [
+		{
+			"Name": "CaptureManagerEditor",
+			"Type": "EditorAndProgram",
+			"ProgramAllowList": [ "LiveLinkHub" ],
+			"LoadingPhase": "Default"
+		},
+		{
+			"Name": "LiveLinkHubExportClient",
+			"Type": "EditorAndProgram",
+			"ProgramAllowList": [ "LiveLinkHub" ],
+			"LoadingPhase": "Default"
+		},
+		{
+			"Name": "LiveLinkFaceMetadata",
+			"Type": "EditorAndProgram",
+			"ProgramAllowList": [ "LiveLinkHub" ],
+			"LoadingPhase": "Default"
+		},
+		{
+			"Name": "StereoCameraMetadata",
+			"Type": "EditorAndProgram",
+			"ProgramAllowList": [ "LiveLinkHub" ],
+			"LoadingPhase": "Default"
+		},
+		{
+			"Name": "LiveLinkHubCaptureSourceWidgets",
+			"Type": "EditorAndProgram",
+			"ProgramAllowList": [ "LiveLinkHub" ],
+			"LoadingPhase": "Default"
+		},
+		{
+			"Name": "LiveLinkHubDiscovery",
+			"Type": "EditorAndProgram",
+			"ProgramAllowList": [ "LiveLinkHub" ],
+			"LoadingPhase": "Default"
+		},
+		{
+			"Name": "CaptureManagerPipeline",
+			"Type": "EditorAndProgram",
+			"ProgramAllowList": [ "LiveLinkHub" ],
+			"LoadingPhase": "Default"
+		},
+		{
+			"Name": "CaptureManagerMediaRW",
+			"Type": "EditorAndProgram",
+			"ProgramAllowList": [ "LiveLinkHub" ],
+			"LoadingPhase": "Default"
+		},
+		{
+			"Name": "CaptureDataConverter",
+			"Type": "EditorAndProgram",
+			"ProgramAllowList": [ "LiveLinkHub" ],
+			"LoadingPhase": "Default"
+		},
+		{
+			"Name": "CaptureCommunicationManager",
+			"Type": "EditorAndProgram",
+			"ProgramAllowList": [ "LiveLinkHub" ],
+			"LoadingPhase": "Default"
+		},
+		{
+			"Name": "GenericCaptureSource",
+			"Type": "EditorAndProgram",
+			"ProgramAllowList": [ "LiveLinkHub" ],
+			"LoadingPhase": "Default"
+		},
+		{
+			"Name": "CaptureManagerSettings",
+			"Type": "EditorAndProgram",
+			"ProgramAllowList": [ "LiveLinkHub" ],
+			"LoadingPhase": "Default"
+		},
+		{
+			"Name": "GenericLiveLinkDevice",
+			"Type": "EditorAndProgram",
+			"ProgramAllowList": [ "LiveLinkHub" ],
+			"LoadingPhase": "Default"
+		},
+		{
+			"Name": "LiveLinkCapabilities",
+			"Type": "EditorAndProgram",
+			"ProgramAllowList": [ "LiveLinkHub" ],
+			"LoadingPhase": "Default"
+		},
+		{
+			"Name": "ExampleLiveLinkDevices",
+			"Type": "EditorAndProgram",
+			"ProgramAllowList": [ "LiveLinkHub" ],
+			"LoadingPhase": "Default"
+		}
+	],
+	"Plugins": [
+		{
+			"Name": "LiveLinkHub",
+			"Enabled": true
+		},
+		{
+			"Name": "ControlFlows",
+			"Enabled": true
+		},
+		{
+			"Name": "CaptureManagerCore",
+			"Enabled": true
+		},
+		{
+			"Name": "CaptureData",
+			"Enabled": true
+		},
+		{
+			"Name": "LiveLinkDevice",
+			"Enabled": true
+		}
+	]
+}
\ No newline at end of file
diff --git a/Engine/Plugins/VirtualProduction/CaptureManager/CaptureManagerApp/Source/CaptureCommunicationManager/CaptureCommunicationManager.Build.cs b/Engine/Plugins/VirtualProduction/CaptureManager/CaptureManagerApp/Source/CaptureCommunicationManager/CaptureCommunicationManager.Build.cs
new file mode 100644
--- /dev/null
+++ b/Engine/Plugins/VirtualProduction/CaptureManager/CaptureManagerApp/Source/CaptureCommunicationManager/CaptureCommunicationManager.Build.cs
@@ -0,0 +1,29 @@
+// Copyright Epic Games, Inc. All Rights Reserved.
+
+using UnrealBuildTool;
+using System.IO;
+
+public class CaptureCommunicationManager : ModuleRules
+{
+	public CaptureCommunicationManager(ReadOnlyTargetRules Target) : base(Target)
+	{
+		PCHUsage = PCHUsageMode.UseExplicitOrSharedPCHs;
+		ShortName = "CaptureCommsManager";
+
+		PublicDependencyModuleNames.AddRange(new string[]
+		{
+			"Core",
+			"CoreUObject",
+			"Engine",
+			"CaptureUtils",
+			"LiveLinkHubCaptureMessaging",
+			"LiveLinkHubExportClient",
+			"CaptureManagerTakeMetadata"
+		});
+
+		PrivateDependencyModuleNames.AddRange(new string[]
+		{
+			"MessagingCommon"
+		});
+	}
+}
diff --git a/Engine/Plugins/VirtualProduction/CaptureManager/CaptureManagerApp/Source/CaptureCommunicationManager/Private/CaptureCommunicationManagerModule.cpp b/Engine/Plugins/VirtualProduction/CaptureManager/CaptureManagerApp/Source/CaptureCommunicationManager/Private/CaptureCommunicationManagerModule.cpp
new file mode 100644
--- /dev/null
+++ b/Engine/Plugins/VirtualProduction/CaptureManager/CaptureManagerApp/Source/CaptureCommunicationManager/Private/CaptureCommunicationManagerModule.cpp
@@ -0,0 +1,20 @@
+// Copyright Epic Games, Inc. All Rights Reserved.
+
+#include "CaptureCommunicationManagerModule.h"
+
+void FCaptureCommunicationManagerModule::StartupModule()
+{
+	CommunicationManager = MakeShared<FCommunicationManager>();
+}
+
+void FCaptureCommunicationManagerModule::ShutdownModule()
+{
+	CommunicationManager = nullptr;
+}
+
+TSharedRef<FCommunicationManager> FCaptureCommunicationManagerModule::GetManager()
+{
+	return CommunicationManager.ToSharedRef();
+}
+
+IMPLEMENT_MODULE(FCaptureCommunicationManagerModule, CaptureCommunicationManager);
\ No newline at end of file
diff --git a/Engine/Plugins/VirtualProduction/CaptureManager/CaptureManagerApp/Source/CaptureCommunicationManager/Public/CaptureCommunicationManagerModule.h b/Engine/Plugins/VirtualProduction/CaptureManager/CaptureManagerApp/Source/CaptureCommunicationManager/Public/CaptureCommunicationManagerModule.h
new file mode 100644
--- /dev/null
+++ b/Engine/Plugins/VirtualProduction/CaptureManager/CaptureManagerApp/Source/CaptureCommunicationManager/Public/CaptureCommunicationManagerModule.h
@@ -0,0 +1,22 @@
+// Copyright Epic Games, Inc. All Rights Reserved.
+
+#pragma once
+
+#include "Modules/ModuleManager.h"
+#include "CaptureCommunicationManager.h"
+
+class CAPTURECOMMUNICATIONMANAGER_API FCaptureCommunicationManagerModule 
+	: public IModuleInterface
+{
+public:
+
+	virtual void StartupModule() override;
+	virtual void ShutdownModule() override;
+
+	TSharedRef<FCommunicationManager> GetManager();
+
+private:
+
+	TSharedPtr<FCommunicationManager> CommunicationManager;
+
+};
diff --git a/Engine/Plugins/VirtualProduction/CaptureManager/CaptureManagerApp/Source/CaptureDataConverter/Private/CaptureDataConverterModule.cpp b/Engine/Plugins/VirtualProduction/CaptureManager/CaptureManagerApp/Source/CaptureDataConverter/Private/CaptureDataConverterModule.cpp
new file mode 100644
--- /dev/null
+++ b/Engine/Plugins/VirtualProduction/CaptureManager/CaptureManagerApp/Source/CaptureDataConverter/Private/CaptureDataConverterModule.cpp
@@ -0,0 +1,93 @@
+// Copyright Epic Games, Inc. All Rights Reserved.
+
+#include "CaptureDataConverterModule.h"
+
+#include "HAL/PlatformProcess.h"
+
+#include "Settings/CaptureManagerSettings.h"
+
+#include "Nodes/ThirdPartyEncoder/CaptureThirdPartyNodeParams.h"
+
+namespace UE::CaptureManager::Private
+{
+
+bool ThirdPartyAvailabilityCheck(const FString& InThirdPartyEncoder)
+{
+	return FPlatformProcess::ExecProcess(*InThirdPartyEncoder, TEXT("-version"), nullptr, nullptr, nullptr);
+}
+
+}
+
+bool FCaptureDataConverterModule::IsThirdPartyEncoderAvailable()
+{
+	return CheckThirdPartyEncoderAvailability();
+}
+
+FString FCaptureDataConverterModule::GetThirdPartyEncoder() const
+{
+	const UCaptureManagerSettings* Settings = GetDefault<UCaptureManagerSettings>();
+
+	if (!Settings->bEnableThirdPartyEncoder)
+	{
+		return FString();
+	}
+
+	return Settings->ThirdPartyEncoder.FilePath;
+}
+
+FString FCaptureDataConverterModule::GetThirdPartyEncoderVideoCommandArguments() const
+{
+	const UCaptureManagerSettings* Settings = GetDefault<UCaptureManagerSettings>();
+
+	if (!Settings->bEnableThirdPartyEncoder)
+	{
+		return FString();
+	}
+
+	return Settings->CustomVideoCommandArguments;
+}
+
+FString FCaptureDataConverterModule::GetThirdPartyEncoderAudioCommandArguments() const
+{
+	const UCaptureManagerSettings* Settings = GetDefault<UCaptureManagerSettings>();
+
+	if (!Settings->bEnableThirdPartyEncoder)
+	{
+		return FString();
+	}
+
+	return Settings->CustomAudioCommandArguments;
+}
+
+void FCaptureDataConverterModule::StartupModule()
+{
+	UCaptureManagerSettings* Settings = GetMutableDefault<UCaptureManagerSettings>();
+
+	if (Settings->CustomAudioCommandArguments.IsEmpty())
+	{
+		Settings->CustomVideoCommandArguments = UE::CaptureManager::VideoCommandArgumentTemplate;
+	}
+	
+	if (Settings->CustomAudioCommandArguments.IsEmpty())
+	{
+		Settings->CustomAudioCommandArguments = UE::CaptureManager::AudioCommandArgumentTemplate;
+	}
+}
+
+void FCaptureDataConverterModule::ShutdownModule()
+{
+}
+
+bool FCaptureDataConverterModule::CheckThirdPartyEncoderAvailability() const
+{
+	const UCaptureManagerSettings* Settings = GetDefault<UCaptureManagerSettings>();
+
+	if (Settings->bEnableThirdPartyEncoder && !Settings->ThirdPartyEncoder.FilePath.IsEmpty())
+	{
+		return UE::CaptureManager::Private::ThirdPartyAvailabilityCheck(Settings->ThirdPartyEncoder.FilePath);
+	}
+
+	return false;
+}
+
+IMPLEMENT_MODULE(FCaptureDataConverterModule, CaptureDataConverter)
\ No newline at end of file
diff --git a/Engine/Plugins/VirtualProduction/CaptureManager/CaptureManagerApp/Source/CaptureDataConverter/Private/CaptureDataConverterModule.h b/Engine/Plugins/VirtualProduction/CaptureManager/CaptureManagerApp/Source/CaptureDataConverter/Private/CaptureDataConverterModule.h
new file mode 100644
--- /dev/null
+++ b/Engine/Plugins/VirtualProduction/CaptureManager/CaptureManagerApp/Source/CaptureDataConverter/Private/CaptureDataConverterModule.h
@@ -0,0 +1,24 @@
+// Copyright Epic Games, Inc. All Rights Reserved.
+
+#pragma once
+
+#include "Modules/ModuleManager.h"
+
+class FCaptureDataConverterModule : public IModuleInterface
+{
+public:
+	
+	bool IsThirdPartyEncoderAvailable();
+	FString GetThirdPartyEncoder() const;
+	FString GetThirdPartyEncoderVideoCommandArguments() const;
+	FString GetThirdPartyEncoderAudioCommandArguments() const;
+
+private:
+
+	virtual void StartupModule() override;
+	virtual void ShutdownModule() override;
+
+	bool CheckThirdPartyEncoderAvailability() const;
+
+	bool bIsThirdPartyEncoderEnabled = false;
+};
\ No newline at end of file
diff --git a/Engine/Plugins/VirtualProduction/CaptureManager/CaptureManagerApp/Source/CaptureDataConverter/Private/Nodes/CaptureConvertDepthData.cpp b/Engine/Plugins/VirtualProduction/CaptureManager/CaptureManagerApp/Source/CaptureDataConverter/Private/Nodes/CaptureConvertDepthData.cpp
new file mode 100644
--- /dev/null
+++ b/Engine/Plugins/VirtualProduction/CaptureManager/CaptureManagerApp/Source/CaptureDataConverter/Private/Nodes/CaptureConvertDepthData.cpp
@@ -0,0 +1,257 @@
+// Copyright Epic Games, Inc. All Rights Reserved.
+
+#include "CaptureConvertDepthData.h"
+
+#include "Async/TaskProgress.h"
+#include "Async/StopToken.h"
+
+#include "CaptureManagerMediaRWModule.h"
+#include "CaptureCopyProgressReporter.h"
+
+#include "IImageWrapperModule.h"
+
+#define LOCTEXT_NAMESPACE "CaptureConvertDepthData"
+
+DEFINE_LOG_CATEGORY_STATIC(LogCaptureConvertDepthData, Log, All);
+
+FCaptureConvertDepthData::FCaptureConvertDepthData(const FTakeMetadata::FVideo& InDepth,
+												   const FString& InOutputDirectory,
+												   const FCaptureConvertDataNodeParams& InParams,
+												   const FCaptureConvertDepthOutputParams& InDepthParams)
+
+	: FConvertDepthNode(InDepth, InOutputDirectory)
+	, Params(InParams)
+	, DepthParams(InDepthParams)
+	, WritingQueue(UE::CaptureManager::TQueueRunner<FWritingContext>::FOnProcess::CreateRaw(this, &FCaptureConvertDepthData::OnWrite))
+{
+}
+
+FCaptureConvertDepthData::FResult FCaptureConvertDepthData::Run()
+{
+	using namespace UE::CaptureManager;
+
+	if (Params.StopToken.IsStopRequested())
+	{
+		FText Message = LOCTEXT("CaptureConvertDepth_AbortedByUser", "Aborted by user");
+		return MakeError(MoveTemp(Message));
+	}
+
+	if (!Depth.PathType.IsSet())
+	{
+		Depth.PathType = FTakeMetadataPathUtils::DetectPathType(Depth.Path);
+		FString PathTypeString = FTakeMetadataPathUtils::PathTypeToString(Depth.PathType.GetValue());
+		UE_LOG(LogCaptureConvertDepthData, Display, TEXT("PathType for %s is unspecified, setting to detected type %s"), *Depth.Path, *PathTypeString);
+	}
+	else
+	{
+		FTakeMetadataPathUtils::ValidatePathType(Depth.Path, Depth.PathType.GetValue());
+	}
+
+	if (Depth.PathType.IsSet() && Depth.PathType.GetValue() == FTakeMetadata::FVideo::EPathType::File)
+	{
+		return ConvertData();
+	}
+	
+	return CopyData();
+}
+
+FCaptureConvertDepthData::FResult FCaptureConvertDepthData::ConvertData()
+{
+	using namespace UE::CaptureManager;
+
+	if (Params.StopToken.IsStopRequested())
+	{
+		FText Message = LOCTEXT("CaptureConvertDepthData_AbortedByUser", "Aborted by user");
+		return MakeError(MoveTemp(Message));
+	}
+
+	FTaskProgress::FTask Task = Params.TaskProgress->StartTask();
+
+	const FString DepthFilePath = FPaths::ConvertRelativePathToFull(Params.TakeOriginDirectory, Depth.Path);
+	FString TargetDirectory = OutputDirectory / Depth.UserId;
+
+	FCaptureManagerMediaRWModule& MediaRWModule =
+		FModuleManager::LoadModuleChecked<FCaptureManagerMediaRWModule>("CaptureManagerMediaRW");
+
+	TUniquePtr<IVideoReader> DepthReader = MediaRWModule.Get().CreateVideoReaderByFormat(Depth.Format);
+	TValueOrError<TUniquePtr<IImageWriter>, FText> ImageWriterResult = MediaRWModule.Get().CreateImageWriter(TargetDirectory, TEXT("depth"), TEXT("exr"));
+
+	if (!DepthReader || ImageWriterResult.HasError())
+	{
+		FText Message = FText::Format(LOCTEXT("CaptureConvertDepthData_UnsupportedFile", "Depth file format is unsupported {0}"), FText::FromString(DepthFilePath));
+		return MakeError(MoveTemp(Message));
+	}
+
+	TOptional<FText> DepthReaderOpen = DepthReader->Open(DepthFilePath);
+	if (DepthReaderOpen.IsSet())
+	{
+		return MakeError(MoveTemp(DepthReaderOpen.GetValue()));
+	}
+
+	TUniquePtr<IImageWriter> ImageWriter = ImageWriterResult.StealValue();
+
+	ON_SCOPE_EXIT
+	{
+		DepthReader->Close();
+		ImageWriter->Close();
+	};
+
+	TPromise<FResult> Promise;
+	TFuture<FResult> Future = Promise.GetFuture();
+	bShouldRun = true;
+
+	while (bShouldRun)
+	{
+		TValueOrError<TUniquePtr<UE::CaptureManager::FMediaTextureSample>, FText> DepthSampleResult = DepthReader->Next();
+
+		FWritingContext Context;
+		Context.ReadResult = MoveTemp(DepthSampleResult);
+		Context.Writer = ImageWriter.Get();
+		
+		if (Depth.FramesCount.IsSet())
+		{
+			Context.TotalDuration = Depth.FramesCount.GetValue();
+		}
+
+		Context.Task = &Task;
+		Context.DonePromise = &Promise;
+
+		WritingQueue.Add(MoveTemp(Context));
+	}
+
+	return Future.Get();
+}
+
+FCaptureConvertDepthData::FResult FCaptureConvertDepthData::CopyData()
+{
+	using namespace UE::CaptureManager;
+
+	const FString DepthFolderPath = FPaths::ConvertRelativePathToFull(Params.TakeOriginDirectory, Depth.Path);
+
+	IFileManager& FileManager = IFileManager::Get();
+	if (!Depth.FramesCount.IsSet() || Depth.FramesCount.GetValue() == 0)
+	{
+		TArray<FString> FoundFiles;
+		FileManager.FindFiles(FoundFiles, *DepthFolderPath, Depth.Format.IsEmpty() ? nullptr : *Depth.Format);
+		Depth.FramesCount = FoundFiles.Num();
+
+		if (!Depth.FramesCount.IsSet() || Depth.FramesCount.GetValue() == 0)
+		{
+			FText ExtensionPostfix = Depth.Format.IsEmpty() ? FText() : FText::Format(LOCTEXT("CaptureConvertDepth_Extension", " with specified extension .{0}"), FText::FromString(Depth.Format));
+			FText Message = FText::Format(LOCTEXT("CaptureConvertDepth_EmptyData", "Copy depth data failed. No depth data found at {0}{1}"), FText::FromString(DepthFolderPath), ExtensionPostfix);
+			UE_LOG(LogCaptureConvertDepthData, Error, TEXT("%s"), *Message.ToString());
+
+			return MakeError(MoveTemp(Message));
+		}
+	}
+
+	FTaskProgress::FTask Task = Params.TaskProgress->StartTask();
+	FTaskProgress TaskProgress(Depth.FramesCount.GetValue(), FTaskProgress::FProgressReporter::CreateLambda([&Task](float InProgress)
+	{
+		Task.Update(InProgress);
+	}));
+
+	FString DestinationDirectory = OutputDirectory / Depth.UserId;
+	FCaptureConvertDepthData::FResult Result = MakeValue();
+	IImageWrapperModule& ImageWrapperModule = FModuleManager::Get().LoadModuleChecked<IImageWrapperModule>(TEXT("ImageWrapper"));
+	FileManager.IterateDirectoryRecursively(*DepthFolderPath, [this, &TaskProgress, &Result, DestinationDirectory, &FileManager, &ImageWrapperModule](const TCHAR* InFileName, bool bInIsDirectory)
+	{
+		if (bInIsDirectory)
+		{
+			return true;
+		}
+
+		if (ImageWrapperModule.GetImageFormatFromExtension(InFileName) != EImageFormat::EXR)
+		{
+			FText Message = LOCTEXT("CaptureConvertDepthData_UnsupportedFileFormat", "Unsupported file format");
+			Result = MakeError(MoveTemp(Message));
+			return false;
+		}
+		
+		FTaskProgress::FTask Task = TaskProgress.StartTask();
+
+		FCopyProgressReporter ProgressReporter(Task, Params.StopToken);
+
+		FString Destination = DestinationDirectory / FPaths::GetCleanFilename(InFileName);
+
+		uint32 CopyResult = FileManager.Copy(*Destination, InFileName, true, true, false, &ProgressReporter);
+
+		if (CopyResult == COPY_Fail)
+		{
+			FText Message = LOCTEXT("CaptureConvertDepthData_CopyFailed", "Failed to copy the file");
+			Result = MakeError(MoveTemp(Message));
+			return false;
+		}
+
+		if (CopyResult == COPY_Canceled)
+		{
+			FText Message = LOCTEXT("CaptureConvertDepthData_AbortedByUser", "Aborted by user");
+			Result = MakeError(MoveTemp(Message));
+			return false;
+		}
+
+		return true;
+	});
+
+	return MakeValue();
+}
+
+void FCaptureConvertDepthData::OnWrite(FWritingContext InContext)
+{
+	if (!bShouldRun)
+	{
+		return;
+	}
+
+	TPromise<FResult>* DonePromise = InContext.DonePromise;
+
+	if (InContext.ReadResult.HasError())
+	{
+		bShouldRun = false;
+
+		DonePromise->SetValue(MakeError(InContext.ReadResult.StealError()));
+		return;
+	}
+
+	TUniquePtr<UE::CaptureManager::FMediaTextureSample> DepthSample = InContext.ReadResult.StealValue();
+
+	// End of stream
+	if (!DepthSample.IsValid())
+	{
+		bShouldRun = false;
+
+		DonePromise->SetValue(MakeValue());
+		return;
+	}
+
+	IImageWriter* Writer = InContext.Writer;
+
+	TOptional<FText> WriterResult = Writer->Append(DepthSample.Get());
+
+	if (WriterResult.IsSet())
+	{
+		bShouldRun = false;
+
+		DonePromise->SetValue(MakeError(MoveTemp(WriterResult.GetValue())));
+
+		return;
+	}
+
+	++CurrentFrame;
+
+	if (InContext.TotalDuration.IsSet() && InContext.TotalDuration.GetValue() > 0)
+	{
+		// We can only update the progress if we know the total duration
+		InContext.Task->Update((float)CurrentFrame / InContext.TotalDuration.GetValue());
+	}
+
+	if (Params.StopToken.IsStopRequested())
+	{
+		bShouldRun = false;
+
+		FText Message = LOCTEXT("LiveLinkConvertDepthNode_AbortedByUser", "Aborted by user");
+		DonePromise->SetValue(MakeError(MoveTemp(Message)));
+	}
+}
+
+#undef LOCTEXT_NAMESPACE
\ No newline at end of file
diff --git a/Engine/Plugins/VirtualProduction/CaptureManager/CaptureManagerApp/Source/CaptureDataConverter/Private/Nodes/CaptureConvertVideoData.cpp b/Engine/Plugins/VirtualProduction/CaptureManager/CaptureManagerApp/Source/CaptureDataConverter/Private/Nodes/CaptureConvertVideoData.cpp
new file mode 100644
--- /dev/null
+++ b/Engine/Plugins/VirtualProduction/CaptureManager/CaptureManagerApp/Source/CaptureDataConverter/Private/Nodes/CaptureConvertVideoData.cpp
@@ -0,0 +1,265 @@
+// Copyright Epic Games, Inc. All Rights Reserved.
+
+#include "CaptureConvertVideoData.h"
+
+#include "Async/TaskProgress.h"
+#include "Async/StopToken.h"
+
+#include "CaptureManagerMediaRWModule.h"
+#include "CaptureCopyProgressReporter.h"
+
+#include "IImageWrapperModule.h"
+
+#define LOCTEXT_NAMESPACE "CaptureConvertVideoData"
+
+DEFINE_LOG_CATEGORY_STATIC(LogCaptureConvertVideoData, Log, All);
+
+FCaptureConvertVideoData::FCaptureConvertVideoData(const FTakeMetadata::FVideo& InVideo, 
+												   const FString& InOutputDirectory, 
+												   const FCaptureConvertDataNodeParams& InParams,
+												   const FCaptureConvertVideoOutputParams& InVideoParams)
+	: FConvertVideoNode(InVideo, InOutputDirectory)
+	, Params(InParams)
+	, VideoParams(InVideoParams)
+	, WritingQueue(UE::CaptureManager::TQueueRunner<FWritingContext>::FOnProcess::CreateRaw(this, &FCaptureConvertVideoData::OnWrite))
+{
+	checkf(!VideoParams.Format.IsEmpty(), TEXT("Output image format must be specified"));
+}
+
+FCaptureConvertVideoData::FResult FCaptureConvertVideoData::Run()
+{
+	using namespace UE::CaptureManager;
+
+	if (Params.StopToken.IsStopRequested())
+	{
+		FText Message = LOCTEXT("CaptureConvertVideo_AbortedByUser", "Video conversion aborted by user");
+		return MakeError(MoveTemp(Message));
+	}
+
+	if (!Video.PathType.IsSet())
+	{
+		Video.PathType = FTakeMetadataPathUtils::DetectPathType(Video.Path);
+		FString PathTypeString = FTakeMetadataPathUtils::PathTypeToString(Video.PathType.GetValue());
+		UE_LOG(LogCaptureConvertVideoData, Display, TEXT("PathType for %s is unspecified, setting to detected type %s"), *Video.Path, *PathTypeString);
+	}
+	else
+	{
+		FTakeMetadataPathUtils::ValidatePathType(Video.Path, Video.PathType.GetValue());
+	}
+
+	if (Video.PathType.IsSet() && Video.PathType.GetValue() == FTakeMetadata::FVideo::EPathType::File)
+	{
+		return ConvertData();
+	}
+	
+	return CopyData();
+}
+
+FCaptureConvertVideoData::FResult FCaptureConvertVideoData::ConvertData()
+{
+	using namespace UE::CaptureManager;
+
+	FTaskProgress::FTask Task = Params.TaskProgress->StartTask();
+
+	const FString VideoFilePath = FPaths::ConvertRelativePathToFull(Params.TakeOriginDirectory, Video.Path);
+	FString DestinationDirectory = OutputDirectory / Video.UserId;
+
+	FCaptureManagerMediaRWModule& MediaRWModule =
+		FModuleManager::LoadModuleChecked<FCaptureManagerMediaRWModule>("CaptureManagerMediaRW");
+
+	TValueOrError<TUniquePtr<IVideoReader>, FText> VideoReaderResult = MediaRWModule.Get().CreateVideoReader(VideoFilePath);
+	TValueOrError<TUniquePtr<IImageWriter>, FText> ImageWriterResult = MediaRWModule.Get().CreateImageWriter(DestinationDirectory, VideoParams.ImageFileName, VideoParams.Format);
+
+	if (VideoReaderResult.HasError() || ImageWriterResult.HasError())
+	{
+		FText Message = FText::Format(LOCTEXT("CaptureConvertVideo_UnsupportedFile", "Video file format is unsupported {0}"), FText::FromString(VideoFilePath));
+		return MakeError(MoveTemp(Message));
+	}
+
+	TUniquePtr<IVideoReader> VideoReader = VideoReaderResult.StealValue();
+	TUniquePtr<IImageWriter> ImageWriter = ImageWriterResult.StealValue();
+
+	ON_SCOPE_EXIT
+	{
+		VideoReader->Close();
+		ImageWriter->Close();
+	};
+
+	const double TotalDuration = VideoReader->GetDuration().GetTotalSeconds();
+
+	TPromise<FResult> Promise;
+	TFuture<FResult> Future = Promise.GetFuture();
+	bShouldRun = true;
+
+	while (bShouldRun)
+	{
+		TValueOrError<TUniquePtr<UE::CaptureManager::FMediaTextureSample>, FText> VideoSampleResult = VideoReader->Next();
+
+		FWritingContext Context;
+		Context.ReadResult = MoveTemp(VideoSampleResult);
+		Context.Writer = ImageWriter.Get();
+		Context.TotalDuration = TotalDuration;
+		Context.Task = &Task;
+		Context.DonePromise = &Promise;
+
+		WritingQueue.Add(MoveTemp(Context));
+	}
+
+	return Future.Get();
+}
+
+FCaptureConvertVideoData::FResult FCaptureConvertVideoData::CopyData()
+{
+	using namespace UE::CaptureManager;
+
+	const FString VideoFolderPath = FPaths::ConvertRelativePathToFull(Params.TakeOriginDirectory, Video.Path);
+
+	IFileManager& FileManager = IFileManager::Get();
+	if (!Video.FramesCount.IsSet() || Video.FramesCount.GetValue() == 0)
+	{
+		TArray<FString> FoundFiles;
+		FileManager.FindFiles(FoundFiles, *VideoFolderPath, Video.Format.IsEmpty() ? nullptr : *Video.Format);
+		
+		if (FoundFiles.IsEmpty())
+		{
+			FText ExtensionPostfix = Video.Format.IsEmpty() ? FText() : FText::Format(LOCTEXT("CaptureConvertVideo_Extension", " with specified extension .{0}"), FText::FromString(Video.Format));
+			FText Message = FText::Format(LOCTEXT("CaptureConvertVideo_EmptyData", "Copy image data failed. No image data found at {0}{1}"), FText::FromString(VideoFolderPath), ExtensionPostfix);
+			UE_LOG(LogCaptureConvertVideoData, Error, TEXT("%s"), *Message.ToString());
+
+			return MakeError(MoveTemp(Message));
+		}
+
+		Video.FramesCount = FoundFiles.Num();
+	}
+
+	FTaskProgress::FTask Task = Params.TaskProgress->StartTask();
+	FTaskProgress TaskProgress(Video.FramesCount.GetValue(), FTaskProgress::FProgressReporter::CreateLambda([&Task](float InProgress)
+	{
+		Task.Update(InProgress);
+	}));
+
+	FString DestinationDirectory = OutputDirectory / Video.UserId;
+
+	FCaptureConvertVideoData::FResult Result = MakeValue();
+
+	IImageWrapperModule& ImageWrapperModule = FModuleManager::Get().LoadModuleChecked<IImageWrapperModule>(TEXT("ImageWrapper"));
+	FileManager.IterateDirectoryRecursively(*VideoFolderPath, [this, &TaskProgress, &Result, DestinationDirectory, &FileManager, &ImageWrapperModule, VideoFolderPath](const TCHAR* InFileName, bool bInIsDirectory)
+	{
+		if (bInIsDirectory)
+		{
+			return true;
+		}
+
+		if (ImageWrapperModule.GetImageFormatFromExtension(InFileName) == EImageFormat::Invalid)
+		{
+			FText Message = FText::Format(LOCTEXT("CaptureConvertVideoData_UnsupportedFileFormat", "Image file format is unsupported {0}"), 
+										  FText::FromString(InFileName));
+			Result = MakeError(MoveTemp(Message));
+			return false;
+		}
+		
+		FTaskProgress::FTask Task = TaskProgress.StartTask();
+
+		FCopyProgressReporter ProgressReporter(Task, Params.StopToken);
+
+		FString Destination = DestinationDirectory / FPaths::GetCleanFilename(InFileName);
+
+		uint32 CopyResult = FileManager.Copy(*Destination, InFileName, true, true, false, &ProgressReporter);
+
+		if (CopyResult == COPY_Fail)
+		{
+			FText Message = FText::Format(LOCTEXT("CaptureConvertVideoData_CopyFailed", "Failed to copy file {0} from {1} to {2}"), FText::FromString(InFileName), FText::FromString(VideoFolderPath), FText::FromString(DestinationDirectory));
+			Result = MakeError(MoveTemp(Message));
+			return false;
+		}
+
+		if (CopyResult == COPY_Canceled)
+		{
+			FText Message = LOCTEXT("CaptureConvertVideoData_AbortedByUser", "Image data copy aborted by user");
+			Result = MakeError(MoveTemp(Message));
+			return false;
+		}
+
+		return true;
+	});
+
+	return MakeValue();
+}
+
+void FCaptureConvertVideoData::OnWrite(FWritingContext InContext)
+{
+	if (!bShouldRun)
+	{
+		return;
+	}
+	
+	TPromise<FResult>* DonePromise = InContext.DonePromise;
+
+	if (InContext.ReadResult.HasError())
+	{
+		bShouldRun = false;
+
+		DonePromise->SetValue(MakeError(InContext.ReadResult.StealError()));
+		return;
+	}
+
+	TUniquePtr<UE::CaptureManager::FMediaTextureSample> VideoSample = InContext.ReadResult.StealValue();
+
+	// End of stream
+	if (!VideoSample.IsValid())
+	{
+		bShouldRun = false;
+		
+		DonePromise->SetValue(MakeValue());
+		return;
+	}
+
+	IImageWriter* Writer = InContext.Writer;
+
+	EMediaOrientation Orientation = ConvertOrientation(Video.Orientation.Get(FTakeMetadata::FVideo::EOrientation::Original));
+	VideoSample->Orientation = Orientation;
+	VideoSample->DesiredFormat = VideoParams.OutputPixelFormat;
+
+	TOptional<FText> WriterResult = Writer->Append(VideoSample.Get());
+
+	if (WriterResult.IsSet())
+	{
+		bShouldRun = false;
+
+		DonePromise->SetValue(MakeError(MoveTemp(WriterResult.GetValue())));
+
+		return;
+	}
+
+	const double Time = VideoSample->Time.GetTotalSeconds();
+	if (InContext.TotalDuration > 0.0)
+	{
+		InContext.Task->Update(Time / InContext.TotalDuration);
+	}
+
+	if (Params.StopToken.IsStopRequested())
+	{
+		bShouldRun = false;
+
+		FText Message = LOCTEXT("CaptureConvertVideo_AbortedByUser", "Video conversion aborted by user");
+		DonePromise->SetValue(MakeError(MoveTemp(Message)));
+	}
+}
+
+EMediaOrientation FCaptureConvertVideoData::ConvertOrientation(FTakeMetadata::FVideo::EOrientation InOrientation) const
+{
+	switch (InOrientation)
+	{
+	case FTakeMetadata::FVideo::EOrientation::CW90:
+			return EMediaOrientation::CW90;
+		case FTakeMetadata::FVideo::EOrientation::CW180:
+			return EMediaOrientation::CW180;
+		case FTakeMetadata::FVideo::EOrientation::CW270:
+			return EMediaOrientation::CW270;
+		case FTakeMetadata::FVideo::EOrientation::Original:
+		default:
+			return EMediaOrientation::Original;
+	}
+}
+
+#undef LOCTEXT_NAMESPACE
\ No newline at end of file
diff --git a/Engine/Plugins/VirtualProduction/CaptureManager/CaptureManagerApp/Source/CaptureDataConverter/Private/Nodes/ThirdPartyEncoder/CaptureConvertAudioDataThirdParty.cpp b/Engine/Plugins/VirtualProduction/CaptureManager/CaptureManagerApp/Source/CaptureDataConverter/Private/Nodes/ThirdPartyEncoder/CaptureConvertAudioDataThirdParty.cpp
new file mode 100644
--- /dev/null
+++ b/Engine/Plugins/VirtualProduction/CaptureManager/CaptureManagerApp/Source/CaptureDataConverter/Private/Nodes/ThirdPartyEncoder/CaptureConvertAudioDataThirdParty.cpp
@@ -0,0 +1,179 @@
+// Copyright Epic Games, Inc. All Rights Reserved.
+
+#include "CaptureConvertAudioDataThirdParty.h"
+
+#include "MediaSample.h"
+
+#include "Async/TaskProgress.h"
+#include "Async/StopToken.h"
+
+#include "Nodes/CaptureCopyProgressReporter.h"
+#include "CaptureThirdPartyNodeUtils.h"
+
+#include "CaptureManagerMediaRWModule.h"
+
+#define LOCTEXT_NAMESPACE "CaptureConvertAudioDataTP"
+
+DEFINE_LOG_CATEGORY_STATIC(LogCaptureConvertAudioDataThirdParty, Log, All);
+
+FCaptureConvertAudioDataThirdParty::FCaptureConvertAudioDataThirdParty(FCaptureThirdPartyNodeParams InThirdPartyEncoder,
+																	   const FTakeMetadata::FAudio& InAudio,
+																	   const FString& InOutputDirectory,
+																	   const FCaptureConvertDataNodeParams& InParams,
+																	   const FCaptureConvertAudioOutputParams& InAudioParams)
+	: FConvertAudioNode(InAudio, InOutputDirectory)
+	, ThirdPartyEncoder(MoveTemp(InThirdPartyEncoder))
+	, Params(InParams)
+	, AudioParams(InAudioParams)
+{
+}
+
+FCaptureConvertAudioDataThirdParty::FResult FCaptureConvertAudioDataThirdParty::Run()
+{
+	using namespace UE::CaptureManager;
+
+	if (Params.StopToken.IsStopRequested())
+	{
+		FText Message = LOCTEXT("CaptureConvertAudioNodeTP_AbortedByUser", "Aborted by user");
+		return MakeError(MoveTemp(Message));
+	}
+
+	if (FPaths::GetExtension(Audio.Path) == AudioParams.Format)
+	{
+		return CopyAudioFile();
+	}
+
+	return ConvertAudioFile();
+}
+
+FCaptureConvertAudioDataThirdParty::FResult FCaptureConvertAudioDataThirdParty::CopyAudioFile()
+{
+	using namespace UE::CaptureManager;
+
+	FTaskProgress::FTask Task = Params.TaskProgress->StartTask();
+
+	FString DestinationDirectory = OutputDirectory / Audio.UserId;
+	const FString AudioFilePath = FPaths::ConvertRelativePathToFull(Params.TakeOriginDirectory, Audio.Path);
+
+	FCopyProgressReporter ProgressReporter(Task, Params.StopToken);
+	FString Destination = DestinationDirectory / FPaths::SetExtension(AudioParams.AudioFileName, AudioParams.Format);
+
+	constexpr bool bReplace = true;
+	constexpr bool bEvenIfReadOnly = true;
+	constexpr bool bAttributes = false;
+	uint32 Result = IFileManager::Get().Copy(*Destination, *AudioFilePath, bReplace, bEvenIfReadOnly, bAttributes, &ProgressReporter);
+
+	if (Result == COPY_Fail)
+	{
+		FText Message = FText::Format(LOCTEXT("CaptureConvertAudioNodeTP_CopyFailed", "Failed to copy the audio file {0}"), 
+									  FText::FromString(AudioFilePath));
+		return MakeError(MoveTemp(Message));
+	}
+
+	if (Result == COPY_Canceled)
+	{
+		FText Message = LOCTEXT("CaptureConvertAudioNodeTP_AbortedByUser", "Aborted by user");
+		return MakeError(MoveTemp(Message));
+	}
+
+	return MakeValue();
+}
+
+FCaptureConvertAudioDataThirdParty::FResult FCaptureConvertAudioDataThirdParty::ConvertAudioFile()
+{
+	using namespace UE::CaptureManager;
+
+	FTaskProgress::FTask Task = Params.TaskProgress->StartTask();
+
+	FString DestinationDirectory = OutputDirectory / Audio.UserId;
+	const FString AudioFilePath = FPaths::ConvertRelativePathToFull(Params.TakeOriginDirectory, Audio.Path);
+
+	const FString AudioOutputFile = FPaths::SetExtension(FPaths::Combine(DestinationDirectory, AudioParams.AudioFileName), AudioParams.Format);
+
+	if (ThirdPartyEncoder.CommandArguments.IsEmpty())
+	{
+		ThirdPartyEncoder.CommandArguments = UE::CaptureManager::AudioCommandArgumentTemplate;
+	}
+
+	FStringFormatNamedArguments FormatArguments;
+	FormatArguments.Add(TEXT("InputFileName"), WrapInQuotes(AudioFilePath));
+	FormatArguments.Add(TEXT("OutputFileName"), WrapInQuotes(AudioOutputFile));
+
+	const FString CommandArguments = FString::Format(*ThirdPartyEncoder.CommandArguments, FormatArguments);
+
+	UE_LOG(LogCaptureConvertAudioDataThirdParty, Display, TEXT("Running the command: %s %s"), *ThirdPartyEncoder.Encoder, *CommandArguments);
+
+	void* ReadPipe = nullptr;
+	void* WritePipe = nullptr;
+	verify(FPlatformProcess::CreatePipe(ReadPipe, WritePipe, false));
+
+	constexpr bool bLaunchDetached = false;
+	constexpr bool bLaunchHidden = true;
+	constexpr bool bLaunchReallyHidden = true;
+	FProcHandle ProcHandle = 
+		FPlatformProcess::CreateProc(*ThirdPartyEncoder.Encoder, *CommandArguments, bLaunchDetached, bLaunchHidden, bLaunchReallyHidden, nullptr, 0, nullptr, WritePipe, nullptr);
+
+	ON_SCOPE_EXIT
+	{
+		if (Params.StopToken.IsStopRequested())
+		{
+			FPlatformProcess::TerminateProc(ProcHandle);
+		}
+
+		FPlatformProcess::CloseProc(ProcHandle);
+		FPlatformProcess::ClosePipe(ReadPipe, WritePipe);
+	};
+
+	if (!ProcHandle.IsValid())
+	{
+		FText Message = FText::Format(LOCTEXT("CaptureConvertAudioNodeTP_ProcessNotFound", "Failed to start the process {0} {1}"),
+									  FText::FromString(ThirdPartyEncoder.Encoder),
+									  FText::FromString(CommandArguments));
+		return MakeError(MoveTemp(Message));
+	}
+
+	TArray<uint8> FullCommandOutput;
+	while (FPlatformProcess::IsProcRunning(ProcHandle))
+	{
+		TArray<uint8> CommandOutput = UE::CaptureManager::ReadPipe(ReadPipe);
+
+		if (CommandOutput.IsEmpty())
+		{
+			FPlatformProcess::Sleep(0.1);
+		}
+
+		FullCommandOutput.Append(MoveTemp(CommandOutput));
+
+		if (Params.StopToken.IsStopRequested())
+		{
+			FText Message = LOCTEXT("CaptureConvertAudioNodeTP_AbortedByUser", "Aborted by user");
+			return MakeError(MoveTemp(Message));
+		}
+	}
+
+	int32 ReturnCode = 0;
+	FPlatformProcess::GetProcReturnCode(ProcHandle, &ReturnCode);
+
+	FullCommandOutput.Append(UE::CaptureManager::ReadPipe(ReadPipe));
+	if (ReturnCode != 0)
+	{
+		if (!FullCommandOutput.IsEmpty())
+		{
+			UE_LOG(LogCaptureConvertAudioDataThirdParty, Error,
+				   TEXT("Failed to run the command: %s %s"), *ThirdPartyEncoder.Encoder, *CommandArguments);
+
+			FString CommandOutputStr = FString::ConstructFromPtrSize(reinterpret_cast<const UTF8CHAR*>(FullCommandOutput.GetData()), FullCommandOutput.Num());
+			UE_LOG(LogCaptureConvertAudioDataThirdParty, Display,
+				   TEXT("Output from the command:\n>>>>>>\n%s<<<<<<"), *CommandOutputStr);
+		}
+
+		FText Message = FText::Format(LOCTEXT("CaptureConvertAudioNodeTP_ErrorRunning", "Error while running the third party encoder (ReturnCode={0})"), 
+									  FText::AsNumber(ReturnCode));
+		return MakeError(MoveTemp(Message));
+	}
+
+	Task.Update(1.0f);
+	return MakeValue();
+}
+
+#undef LOCTEXT_NAMESPACE
\ No newline at end of file
diff --git a/Engine/Plugins/VirtualProduction/CaptureManager/CaptureManagerApp/Source/CaptureDataConverter/Private/Nodes/ThirdPartyEncoder/CaptureConvertAudioDataThirdParty.h b/Engine/Plugins/VirtualProduction/CaptureManager/CaptureManagerApp/Source/CaptureDataConverter/Private/Nodes/ThirdPartyEncoder/CaptureConvertAudioDataThirdParty.h
new file mode 100644
--- /dev/null
+++ b/Engine/Plugins/VirtualProduction/CaptureManager/CaptureManagerApp/Source/CaptureDataConverter/Private/Nodes/ThirdPartyEncoder/CaptureConvertAudioDataThirdParty.h
@@ -0,0 +1,32 @@
+// Copyright Epic Games, Inc. All Rights Reserved.
+
+#pragma once
+
+#include "Nodes/ConvertAudioNode.h"
+
+#include "CaptureThirdPartyNodeParams.h"
+
+#include "CaptureDataConverterNodeParams.h"
+
+class FCaptureConvertAudioDataThirdParty final :
+	public FConvertAudioNode
+{
+public:
+
+	FCaptureConvertAudioDataThirdParty(FCaptureThirdPartyNodeParams InThirdPartEncoder,
+									   const FTakeMetadata::FAudio& InAudio,
+									   const FString& InOutputDirectory,
+									   const FCaptureConvertDataNodeParams& InParams,
+									   const FCaptureConvertAudioOutputParams& InAudioParams);
+
+private:
+
+	virtual FResult Run() override;
+
+	FResult CopyAudioFile();
+	FResult ConvertAudioFile();
+
+	FCaptureThirdPartyNodeParams ThirdPartyEncoder;
+	FCaptureConvertDataNodeParams Params;
+	FCaptureConvertAudioOutputParams AudioParams;
+};
\ No newline at end of file
diff --git a/Engine/Plugins/VirtualProduction/CaptureManager/CaptureManagerApp/Source/CaptureDataConverter/Private/Nodes/ThirdPartyEncoder/CaptureConvertVideoDataThirdParty.cpp b/Engine/Plugins/VirtualProduction/CaptureManager/CaptureManagerApp/Source/CaptureDataConverter/Private/Nodes/ThirdPartyEncoder/CaptureConvertVideoDataThirdParty.cpp
new file mode 100644
--- /dev/null
+++ b/Engine/Plugins/VirtualProduction/CaptureManager/CaptureManagerApp/Source/CaptureDataConverter/Private/Nodes/ThirdPartyEncoder/CaptureConvertVideoDataThirdParty.cpp
@@ -0,0 +1,334 @@
+// Copyright Epic Games, Inc. All Rights Reserved.
+
+#include "CaptureConvertVideoDataThirdParty.h"
+
+#include "MediaSample.h"
+
+#include "Async/TaskProgress.h"
+#include "Async/StopToken.h"
+
+#include "CaptureManagerMediaRWModule.h"
+#include "Nodes/CaptureCopyProgressReporter.h"
+#include "CaptureThirdPartyNodeUtils.h"
+
+#include "IImageWrapperModule.h"
+
+#define LOCTEXT_NAMESPACE "CaptureConvertVideoDataTP"
+
+DEFINE_LOG_CATEGORY_STATIC(LogCaptureConvertVideoDataThirdParty, Log, All);
+
+namespace UE::CaptureManager::Private
+{
+
+static void ConvertOrientation(FTakeMetadata::FVideo::EOrientation InOrientation, TArray<FString>& OutFilterArray)
+{
+	if (InOrientation == FTakeMetadata::FVideo::EOrientation::CW90)
+	{
+		OutFilterArray.Add(TEXT("transponse=cclock"));
+	}
+	else if (InOrientation == FTakeMetadata::FVideo::EOrientation::CW180)
+	{
+		OutFilterArray.Add(TEXT("transponse=clock,transponse=clock"));
+	}
+	else if (InOrientation == FTakeMetadata::FVideo::EOrientation::CW270)
+	{
+		OutFilterArray.Add(TEXT("transpose=clock"));
+	}
+}
+
+static void ConvertPixelFormat(EMediaTexturePixelFormat InPixelFormat, TArray<FString>& OutFilterArray)
+{
+	if (InPixelFormat == EMediaTexturePixelFormat::U8_RGB)
+	{
+		OutFilterArray.Add(TEXT("format=rgb0"));
+	}
+	else if (InPixelFormat == EMediaTexturePixelFormat::U8_BGR)
+	{
+		OutFilterArray.Add(TEXT("format=bgr0"));
+	}
+	else if (InPixelFormat == EMediaTexturePixelFormat::U8_RGBA)
+	{
+		OutFilterArray.Add(TEXT("format=rgba"));
+	}
+	else if (InPixelFormat == EMediaTexturePixelFormat::U8_BGRA)
+	{
+		OutFilterArray.Add(TEXT("format=bgra"));
+	}
+	else if (InPixelFormat == EMediaTexturePixelFormat::U8_I444)
+	{
+		OutFilterArray.Add(TEXT("format=yuvj444p"));
+	}
+	else if (InPixelFormat == EMediaTexturePixelFormat::U8_I420)
+	{
+		OutFilterArray.Add(TEXT("format=yuvj420p"));
+	}
+	else if (InPixelFormat == EMediaTexturePixelFormat::U8_YUY2)
+	{
+		OutFilterArray.Add(TEXT("format=yuyv422"));
+	}
+	else if (InPixelFormat == EMediaTexturePixelFormat::U8_NV12)
+	{
+		OutFilterArray.Add(TEXT("format=nv12"));
+	}
+	else if (InPixelFormat == EMediaTexturePixelFormat::U8_Mono)
+	{
+		OutFilterArray.Add(TEXT("format=gray"));
+	}
+	else if (InPixelFormat == EMediaTexturePixelFormat::U16_Mono)
+	{
+		OutFilterArray.Add(TEXT("format=gray16be"));
+	}
+	else if (InPixelFormat == EMediaTexturePixelFormat::F_Mono)
+	{
+		OutFilterArray.Add(TEXT("format=gray32fbe"));
+	}
+}
+
+static FString CreateConversionArguments(EMediaTexturePixelFormat InPixelFormat, FTakeMetadata::FVideo::EOrientation InOrientation)
+{
+	TArray<FString> FilterArray;
+
+	ConvertPixelFormat(InPixelFormat, FilterArray);
+	ConvertOrientation(InOrientation, FilterArray);
+
+	if (FilterArray.IsEmpty())
+	{
+		return FString();
+	}
+
+	FString Filters = FString::Join(FilterArray, TEXT(","));
+
+	return FString::Format(TEXT("-vf \"{0}\""), { Filters });
+}
+
+}
+
+FCaptureConvertVideoDataThirdParty::FCaptureConvertVideoDataThirdParty(FCaptureThirdPartyNodeParams InThirdPartyEncoder,
+																	   const FTakeMetadata::FVideo& InVideo,
+																	   const FString& InOutputDirectory, 
+																	   const FCaptureConvertDataNodeParams& InParams,
+																	   const FCaptureConvertVideoOutputParams& InVideoParams)
+	: FConvertVideoNode(InVideo, InOutputDirectory)
+	, ThirdPartyEncoder(MoveTemp(InThirdPartyEncoder))
+	, Params(InParams)
+	, VideoParams(InVideoParams)
+{
+	checkf(!VideoParams.Format.IsEmpty(), TEXT("Format MUST be specified"));
+}
+
+FCaptureConvertVideoDataThirdParty::FResult FCaptureConvertVideoDataThirdParty::Run()
+{
+	using namespace UE::CaptureManager;
+
+	if (Params.StopToken.IsStopRequested())
+	{
+		FText Message = LOCTEXT("CaptureConvertVideoNodeTP_AbortedByUser", "Aborted by user");
+		return MakeError(MoveTemp(Message));
+	}
+
+	if (!Video.PathType.IsSet())
+	{
+		Video.PathType = FTakeMetadataPathUtils::DetectPathType(Video.Path);
+		FString PathTypeString = FTakeMetadataPathUtils::PathTypeToString(Video.PathType.GetValue());
+		UE_LOG(LogCaptureConvertVideoDataThirdParty, Display, TEXT("PathType for %s is unspecified, setting to detected type %s"), *Video.Path, *PathTypeString);
+	}
+	else
+	{
+		FTakeMetadataPathUtils::ValidatePathType(Video.Path, Video.PathType.GetValue());
+	}
+
+	if (Video.PathType.IsSet() && Video.PathType.GetValue() == FTakeMetadata::FVideo::EPathType::File)
+	{
+		return ConvertData();
+	}
+	
+	return CopyData();
+}
+
+FCaptureConvertVideoDataThirdParty::FResult FCaptureConvertVideoDataThirdParty::ConvertData()
+{
+	using namespace UE::CaptureManager;
+
+	FTaskProgress::FTask Task = Params.TaskProgress->StartTask();
+
+	const FString VideoFilePath = FPaths::ConvertRelativePathToFull(Params.TakeOriginDirectory, Video.Path);
+	FString DestinationDirectory = OutputDirectory / Video.UserId;
+
+	const FString ImageFileName = FString::Format(TEXT("{0}_%06d.{1}"), { VideoParams.ImageFileName, VideoParams.Format });
+	const FString ImageFilePath = FPaths::Combine(DestinationDirectory, ImageFileName);
+
+	FTakeMetadata::FVideo::EOrientation Orientation = Video.Orientation.Get(FTakeMetadata::FVideo::EOrientation::Original);
+
+	FString ConversionArguments = Private::CreateConversionArguments(VideoParams.OutputPixelFormat, Orientation);
+
+	if (ThirdPartyEncoder.CommandArguments.IsEmpty())
+	{
+		ThirdPartyEncoder.CommandArguments = UE::CaptureManager::VideoCommandArgumentTemplate;
+	}
+
+	FStringFormatNamedArguments FormatArguments;
+	FormatArguments.Add(TEXT("InputFileName"), WrapInQuotes(VideoFilePath));
+	FormatArguments.Add(TEXT("OutputFileName"), WrapInQuotes(ImageFilePath));
+	FormatArguments.Add(TEXT("ConversionFilters"), ConversionArguments);
+
+	FString CommandArguments = FString::Format(*ThirdPartyEncoder.CommandArguments, FormatArguments);
+	
+	UE_LOG(LogCaptureConvertVideoDataThirdParty, Display, TEXT("Running the command: %s %s"), *ThirdPartyEncoder.Encoder, *CommandArguments);
+
+	void* ReadPipe = nullptr;
+	void* WritePipe = nullptr;
+	verify(FPlatformProcess::CreatePipe(ReadPipe, WritePipe, false));
+
+	constexpr bool bLaunchDetached = false;
+	constexpr bool bLaunchHidden = true;
+	constexpr bool bLaunchReallyHidden = true;
+	FProcHandle ProcHandle = 
+		FPlatformProcess::CreateProc(*ThirdPartyEncoder.Encoder, *CommandArguments, bLaunchDetached, bLaunchHidden, bLaunchReallyHidden, nullptr, 0, nullptr, WritePipe, nullptr);
+
+	ON_SCOPE_EXIT
+	{
+		if (Params.StopToken.IsStopRequested())
+		{
+			FPlatformProcess::TerminateProc(ProcHandle);
+		}
+
+		FPlatformProcess::CloseProc(ProcHandle);
+		FPlatformProcess::ClosePipe(ReadPipe, WritePipe);
+	};
+
+	if (!ProcHandle.IsValid())
+	{
+		FText Message = FText::Format(LOCTEXT("CaptureConvertVideoNodeTP_ProcessNotFound", "Failed to start the process {0} {1}"),
+									  FText::FromString(ThirdPartyEncoder.Encoder),
+									  FText::FromString(CommandArguments));
+		return MakeError(MoveTemp(Message));
+	}
+
+	TArray<uint8> FullCommandOutput;
+	while (FPlatformProcess::IsProcRunning(ProcHandle))
+	{
+		TArray<uint8> CommandOutput = UE::CaptureManager::ReadPipe(ReadPipe);
+
+		if (CommandOutput.IsEmpty())
+		{
+			FPlatformProcess::Sleep(0.1);
+		}
+
+		FullCommandOutput.Append(MoveTemp(CommandOutput));
+		
+		if (Params.StopToken.IsStopRequested())
+		{
+			FText Message = LOCTEXT("CaptureConvertVideoNodeTP_AbortedByUser", "Aborted by user");
+			return MakeError(MoveTemp(Message));
+		}
+	}
+
+	int32 ReturnCode = 0;
+	FPlatformProcess::GetProcReturnCode(ProcHandle, &ReturnCode);
+
+	FullCommandOutput.Append(UE::CaptureManager::ReadPipe(ReadPipe));
+
+	if (ReturnCode != 0)
+	{
+		if (!FullCommandOutput.IsEmpty())
+		{
+			UE_LOG(LogCaptureConvertVideoDataThirdParty, Error,
+				   TEXT("Failed to run the command: %s %s"), *ThirdPartyEncoder.Encoder, *CommandArguments);
+
+			FString CommandOutputStr = FString::ConstructFromPtrSize(reinterpret_cast<const UTF8CHAR*>(FullCommandOutput.GetData()), FullCommandOutput.Num());
+			UE_LOG(LogCaptureConvertVideoDataThirdParty, Display,
+				   TEXT("Output from the command:\n>>>>>>\n%s<<<<<<"), *CommandOutputStr);
+		}
+
+		FText Message = FText::Format(LOCTEXT("CaptureConvertVideoNodeTP_ErrorRunning", "Error while running the third party encoder (ReturnCode={0})"),
+									  FText::AsNumber(ReturnCode));
+		return MakeError(MoveTemp(Message));
+	}
+
+	Task.Update(1.0f);
+	return MakeValue();
+}
+
+FCaptureConvertVideoDataThirdParty::FResult FCaptureConvertVideoDataThirdParty::CopyData()
+{
+	using namespace UE::CaptureManager;
+
+	const FString VideoFolderPath = FPaths::ConvertRelativePathToFull(Params.TakeOriginDirectory, Video.Path);
+
+	IFileManager& FileManager = IFileManager::Get();
+	if (!Video.FramesCount.IsSet() || Video.FramesCount.GetValue() == 0)
+	{
+		TArray<FString> FoundFiles;
+		FileManager.FindFiles(FoundFiles, *VideoFolderPath, Video.Format.IsEmpty() ? nullptr : *Video.Format);
+
+		if (FoundFiles.IsEmpty())
+		{
+			FText ExtensionPostfix = Video.Format.IsEmpty() ? FText() : FText::Format(LOCTEXT("CaptureConvertVideo_Extension", " with specified extension .{0}"), FText::FromString(Video.Format));
+			FText Message = FText::Format(LOCTEXT("CaptureConvertVideoNodeTP_EmptyData", "Copy image data failed. No image data found at {0}{1}"), FText::FromString(VideoFolderPath), ExtensionPostfix);
+			UE_LOG(LogCaptureConvertVideoDataThirdParty, Error, TEXT("%s"), *Message.ToString());
+
+			return MakeError(MoveTemp(Message));
+		}
+
+		Video.FramesCount = FoundFiles.Num();
+	}
+
+	FTaskProgress::FTask Task = Params.TaskProgress->StartTask();
+	FTaskProgress TaskProgress(Video.FramesCount.GetValue(), FTaskProgress::FProgressReporter::CreateLambda([&Task](float InProgress)
+	{
+		Task.Update(InProgress);
+	}));
+
+	FString DestinationDirectory = OutputDirectory / Video.UserId;
+
+	FCaptureConvertVideoDataThirdParty::FResult Result = MakeValue();
+
+	IImageWrapperModule& ImageWrapperModule = FModuleManager::Get().LoadModuleChecked<IImageWrapperModule>(TEXT("ImageWrapper"));
+	FileManager.IterateDirectoryRecursively(*VideoFolderPath, [this, &TaskProgress, &Result, DestinationDirectory, &FileManager, &ImageWrapperModule](const TCHAR* InFileName, bool bInIsDirectory)
+	{
+		if (bInIsDirectory)
+		{
+			return true;
+		}
+
+		if (ImageWrapperModule.GetImageFormatFromExtension(InFileName) == EImageFormat::Invalid)
+		{
+			FText Message = FText::Format(LOCTEXT("CaptureConvertVideoDataTP_UnsupportedFileFormat", "Image file format is unsupported {0}"), 
+										  FText::FromString(InFileName));
+			Result = MakeError(MoveTemp(Message));
+			return false;
+		}
+		
+		FTaskProgress::FTask Task = TaskProgress.StartTask();
+
+		FCopyProgressReporter ProgressReporter(Task, Params.StopToken);
+
+		FString Destination = DestinationDirectory / FPaths::GetCleanFilename(InFileName);
+
+		constexpr bool bReplace = true;
+		constexpr bool bEvenIfReadOnly = true;
+		constexpr bool bAttributes = false;
+		uint32 CopyResult = FileManager.Copy(*Destination, InFileName, bReplace, bEvenIfReadOnly, bAttributes, &ProgressReporter);
+
+		if (CopyResult == COPY_Fail)
+		{
+			FText Message = FText::Format(LOCTEXT("CaptureConvertVideoNodeTP_CopyFailed", "Failed to copy the video file {0}"),
+										  FText::FromString(InFileName));
+			Result = MakeError(MoveTemp(Message));
+			return false;
+		}
+
+		if (CopyResult == COPY_Canceled)
+		{
+			FText Message = LOCTEXT("CaptureConvertVideoNodeTP_AbortedByUser", "Aborted by user");
+			Result = MakeError(MoveTemp(Message));
+			return false;
+		}
+
+		return true;
+	});
+
+	return MakeValue();
+}
+
+#undef LOCTEXT_NAMESPACE
\ No newline at end of file
diff --git a/Engine/Plugins/VirtualProduction/CaptureManager/CaptureManagerApp/Source/CaptureDataConverter/Private/Nodes/ThirdPartyEncoder/CaptureThirdPartyNodeParams.h b/Engine/Plugins/VirtualProduction/CaptureManager/CaptureManagerApp/Source/CaptureDataConverter/Private/Nodes/ThirdPartyEncoder/CaptureThirdPartyNodeParams.h
new file mode 100644
--- /dev/null
+++ b/Engine/Plugins/VirtualProduction/CaptureManager/CaptureManagerApp/Source/CaptureDataConverter/Private/Nodes/ThirdPartyEncoder/CaptureThirdPartyNodeParams.h
@@ -0,0 +1,19 @@
+// Copyright Epic Games, Inc. All Rights Reserved.
+
+#pragma once
+
+#include "Containers/UnrealString.h"
+
+namespace UE::CaptureManager
+{
+
+static const FString VideoCommandArgumentTemplate = TEXT("-noautorotate -i {InputFileName} -an {ConversionFilters} -q:v 1 -qmax 1 -qmin 1 -start_number 0 {OutputFileName}");
+static const FString AudioCommandArgumentTemplate = TEXT("-i {InputFileName} -vn {OutputFileName}");
+
+}
+
+struct FCaptureThirdPartyNodeParams
+{
+	FString Encoder;
+	FString CommandArguments;
+};
diff --git a/Engine/Plugins/VirtualProduction/CaptureManager/CaptureManagerApp/Source/CaptureDataConverter/Public/CaptureDataConverterError.h b/Engine/Plugins/VirtualProduction/CaptureManager/CaptureManagerApp/Source/CaptureDataConverter/Public/CaptureDataConverterError.h
new file mode 100644
--- /dev/null
+++ b/Engine/Plugins/VirtualProduction/CaptureManager/CaptureManagerApp/Source/CaptureDataConverter/Public/CaptureDataConverterError.h
@@ -0,0 +1,19 @@
+// Copyright Epic Games, Inc. All Rights Reserved.
+
+#pragma once
+
+#include "Internationalization/Text.h"
+
+class CAPTUREDATACONVERTER_API FCaptureDataConverterError
+{
+public:
+
+	explicit FCaptureDataConverterError(TArray<FText> InErrors);
+	~FCaptureDataConverterError();
+
+	TArray<FText> GetErrors() const;
+
+private:
+
+	TArray<FText> Errors;
+};
diff --git a/Engine/Plugins/VirtualProduction/CaptureManager/CaptureManagerApp/Source/CaptureDataConverter/Public/Nodes/CaptureConvertParams.h b/Engine/Plugins/VirtualProduction/CaptureManager/CaptureManagerApp/Source/CaptureDataConverter/Public/Nodes/CaptureConvertParams.h
new file mode 100644
--- /dev/null
+++ b/Engine/Plugins/VirtualProduction/CaptureManager/CaptureManagerApp/Source/CaptureDataConverter/Public/Nodes/CaptureConvertParams.h
@@ -0,0 +1,22 @@
+// Copyright Epic Games, Inc. All Rights Reserved.
+
+#pragma once
+
+#include "CaptureManagerPipelineNode.h"
+
+#include "CaptureDataConverterNodeParams.h"
+
+class CAPTUREDATACONVERTER_API FCaptureConvertParams
+{
+public:
+
+	FCaptureConvertParams();
+	virtual ~FCaptureConvertParams();
+
+	void SetParams(const FCaptureConvertDataNodeParams& InParams);
+
+protected:
+
+	FCaptureConvertDataNodeParams Params;
+	
+};
\ No newline at end of file
diff --git a/Engine/Plugins/VirtualProduction/CaptureManager/CaptureManagerApp/Source/CaptureManagerEditor/Private/IngestManagement/IngestJobExecutor.cpp b/Engine/Plugins/VirtualProduction/CaptureManager/CaptureManagerApp/Source/CaptureManagerEditor/Private/IngestManagement/IngestJobExecutor.cpp
new file mode 100644
--- /dev/null
+++ b/Engine/Plugins/VirtualProduction/CaptureManager/CaptureManagerApp/Source/CaptureManagerEditor/Private/IngestManagement/IngestJobExecutor.cpp
@@ -0,0 +1,219 @@
+// Copyright Epic Games, Inc. All Rights Reserved.
+
+#include "IngestJobExecutor.h"
+
+#include "IngestJobQueue.h"
+
+#include "Engine/Engine.h"
+
+#include "Ingest/LiveLinkDeviceCapability_Ingest.h"
+#include "LiveLinkDeviceSubsystem.h"
+#include "LiveLinkDevice.h"
+
+DEFINE_LOG_CATEGORY_STATIC(LogCaptureManagerIngestJobExecutor, Log, All);
+
+namespace UE::CaptureManager
+{
+
+EIngestCapability_ImagePixelFormat ConvertPixelFormat(::EImagePixelFormat InImagePixelFormat)
+{
+	switch (InImagePixelFormat)
+	{
+		case ::EImagePixelFormat::U8_Mono:
+			return EIngestCapability_ImagePixelFormat::U8_Mono;
+		case ::EImagePixelFormat::U8_BGRA:
+		default:
+			return EIngestCapability_ImagePixelFormat::U8_BGRA;
+	}
+}
+
+FIngestJobExecutor::FIngestJobExecutor(FString InExecutorName, TSharedRef<FIngestJobQueue> InProcessingQueue, FOnComplete InOnComplete, FJobProcessingStateChanged InJobProcessingStateChanged) :
+	ExecutorName(MoveTemp(InExecutorName)),
+	ProcessingQueue(MoveTemp(InProcessingQueue)),
+	Complete(MoveTemp(InOnComplete)),
+	JobProcessingStateChanged(MoveTemp(InJobProcessingStateChanged)),
+	bIsRunning(false),
+	bStopRequested(false),
+	Thread(TUniquePtr<FRunnableThread>(FRunnableThread::Create(this, *ExecutorName)))
+{
+}
+
+FIngestJobExecutor::~FIngestJobExecutor()
+{
+	if (Thread)
+	{
+		Thread->Kill(true);
+	}
+}
+
+bool FIngestJobExecutor::Init()
+{
+	bIsRunning = true;
+	return true;
+}
+
+uint32 FIngestJobExecutor::Run()
+{
+	while (!bStopRequested)
+	{
+		TPromise<void> Promise;
+		TFuture<void> Future = Promise.GetFuture();
+
+		TSharedPtr<FIngestJob> Job = ProcessingQueue->ClaimFirstPending(FIngestJob::EProcessingState::Running);
+
+		if (Job)
+		{
+			check(Job->GetProcessingState() != FIngestJob::EProcessingState::Pending);
+			JobProcessingStateChanged.ExecuteIfBound(Job->GetGuid(), Job->GetProcessingState());
+
+			// Note: This may be a blocking call depending on how the capture device has been implemented. If so, this 
+			// may take a while to stop (no way to cancel in that case).
+			UIngestCapability_ProcessHandle* TakeProcessHandle = RunIngest(*Job, Promise);
+
+			FScopeLock Lock(&CriticalSection);
+			check(!IngestProcessHandle);
+
+			IngestProcessHandle = MoveTemp(TakeProcessHandle);
+			
+			// There's a period between RunIngest and the assignment of the member handle, during which Stop() may have
+			// been called. If so, immediately terminate (as Cancel() will not be called otherwise).
+			if (bStopRequested && !Future.IsReady() && IngestProcessHandle)
+			{
+				IngestProcessHandle->Cancel();
+			}
+		}
+		else
+		{
+			// No job found, so go back to watching the queue
+			Promise.SetValue();
+		}
+
+		Future.Wait();
+
+		{
+			FScopeLock Lock(&CriticalSection);
+
+			if (IngestProcessHandle)
+			{
+				IngestProcessHandle = nullptr;
+			}
+		}
+
+		FPlatformProcess::Sleep(1.0);
+	}
+
+	return 0;
+}
+
+void FIngestJobExecutor::Stop()
+{
+	bStopRequested = true;
+
+	FScopeLock Lock(&CriticalSection);
+
+	if (IngestProcessHandle)
+	{
+		IngestProcessHandle->Cancel();
+	}
+}
+
+void FIngestJobExecutor::Exit()
+{
+	bIsRunning = false;
+	Complete.ExecuteIfBound();
+}
+
+UIngestCapability_ProcessHandle* FIngestJobExecutor::RunIngest(FIngestJob& InJob, TPromise<void>& InPromise)
+{
+	ULiveLinkDeviceSubsystem* DeviceSubsystem = GEngine->GetEngineSubsystem<ULiveLinkDeviceSubsystem>();
+
+	const TObjectPtr<ULiveLinkDevice>* DevicePtr = DeviceSubsystem->GetDeviceMap().Find(InJob.GetCaptureDeviceId());
+
+	if (!DevicePtr)
+	{
+		UE_LOG(
+			LogCaptureManagerIngestJobExecutor,
+			Error,
+			TEXT("[%s] Ingest failed, Capture Device does not exist"), *ExecutorName
+		);
+
+		// Set the processing state to anything other than Pending, to prevent other executors picking up the same bad job.
+		InJob.SetProcessingState(FIngestJob::EProcessingState::Aborted);
+		InPromise.SetValue();
+		return nullptr;
+	}
+
+	TObjectPtr<ULiveLinkDevice> Device = *DevicePtr;
+
+	if (!Device->Implements<ULiveLinkDeviceCapability_Ingest>())
+	{
+		UE_LOG(
+			LogCaptureManagerIngestJobExecutor,
+			Error,
+			TEXT("[%s] Ingest failed, Capture Device does not support the ingest capability: %s"), *ExecutorName, *Device.GetName()
+		);
+
+		// Set the processing state to anything other than Pending, to prevent other executors picking up the same bad job.
+		InJob.SetProcessingState(FIngestJob::EProcessingState::Aborted);
+		InPromise.SetValue();
+		return nullptr;
+	}
+
+	FIngestProcessFinishReporter ProcessFinished = FIngestProcessFinishReporter::Type::CreateLambda(
+		[this, &InPromise, &InJob](FTakeId TakeId, TValueOrError<void, FIngestCapability_Error> IngestError)
+		{
+			FIngestJob::EProcessingState ProcessingState;
+
+			if (IngestError.HasValue())
+			{
+				ProcessingState = FIngestJob::EProcessingState::Complete;
+				InJob.SetProgress(1.0);
+			}
+			else
+			{
+				ProcessingState = FIngestJob::EProcessingState::Aborted;
+				UE_LOG(LogCaptureManagerIngestJobExecutor, Error, TEXT("[%s] Ingest failed: %s (Error code: %d)"), *ExecutorName, *IngestError.GetError().GetMessage(), IngestError.GetError().GetCode())
+			}
+
+			InJob.SetProcessingState(ProcessingState);
+			JobProcessingStateChanged.ExecuteIfBound(InJob.GetGuid(), ProcessingState);
+			InPromise.SetValue();
+		}
+	);
+
+	FIngestProcessProgressReporter ProgressReporter = FIngestProcessProgressReporter::Type::CreateLambda(
+		[&InJob, LastProgress = 0.0](FTakeId TakeId, const double Progress) mutable
+		{
+			// A little sanity check that we're receiving updates in the expected order.
+			check(Progress >= LastProgress);
+			LastProgress = Progress;
+
+			InJob.SetProgress(Progress);
+		}
+	);
+
+	UIngestCapability_ProcessHandle* ProcessHandle =
+		ILiveLinkDeviceCapability_Ingest::Execute_PrepareIngestProcess(Device, InJob.GetTakeId());
+
+	ProcessHandle->ProcessFinishedReporter = MoveTemp(ProcessFinished);
+	ProcessHandle->ProcessProgressReporter = MoveTemp(ProgressReporter);
+
+	const FIngestJob::FSettings Settings = InJob.GetSettings();
+
+	UIngestCapability_Options* IngestOptions = NewObject<UIngestCapability_Options>();
+	IngestOptions->Location = Settings.LocalWorkFolder;
+	FIngestCapability_VideoOptions VideoOptions;
+	VideoOptions.FileNamePrefix = Settings.VideoSettings.FileNamePrefix;
+	VideoOptions.Format = Settings.VideoSettings.Format;
+	VideoOptions.PixelFormat = ConvertPixelFormat(Settings.VideoSettings.ImagePixelFormat);
+	IngestOptions->Video = VideoOptions;
+	FIngestCapability_AudioOptions AudioOptions;
+	AudioOptions.FileNamePrefix = Settings.AudioSettings.FileNamePrefix;
+	AudioOptions.Format = Settings.AudioSettings.Format;
+	IngestOptions->Audio = AudioOptions;
+
+	ProcessHandle->Run(IngestOptions);
+	return ProcessHandle;
+}
+
+} // namespace UE::CaptureManager
diff --git a/Engine/Plugins/VirtualProduction/CaptureManager/CaptureManagerApp/Source/CaptureManagerEditor/Private/IngestManagement/IngestJobExecutor.h b/Engine/Plugins/VirtualProduction/CaptureManager/CaptureManagerApp/Source/CaptureManagerEditor/Private/IngestManagement/IngestJobExecutor.h
new file mode 100644
--- /dev/null
+++ b/Engine/Plugins/VirtualProduction/CaptureManager/CaptureManagerApp/Source/CaptureManagerEditor/Private/IngestManagement/IngestJobExecutor.h
@@ -0,0 +1,42 @@
+// Copyright Epic Games, Inc. All Rights Reserved.
+
+#pragma once
+
+#include "IngestJob.h"
+
+class UIngestCapability_ProcessHandle;
+
+namespace UE::CaptureManager
+{
+class FIngestJobQueue;
+
+class FIngestJobExecutor : public FRunnable
+{
+public:
+	DECLARE_DELEGATE(FOnComplete);
+	DECLARE_DELEGATE_TwoParams(FJobProcessingStateChanged, FGuid, FIngestJob::EProcessingState);
+
+	FIngestJobExecutor(FString InExecutorName, TSharedRef<FIngestJobQueue> InProcessingQueue, FOnComplete InOnComplete, FJobProcessingStateChanged InJobProcessingStateChanged);
+	~FIngestJobExecutor();
+
+	virtual bool Init() override;
+	virtual uint32 Run() override;
+	virtual void Stop() override;
+	virtual void Exit() override;
+
+private:
+	UIngestCapability_ProcessHandle* RunIngest(FIngestJob& InJob, TPromise<void>& InPromise);
+
+	FCriticalSection CriticalSection;
+	FString ExecutorName;
+	TSharedRef<FIngestJobQueue> ProcessingQueue;
+	UIngestCapability_ProcessHandle* IngestProcessHandle = nullptr;
+	FOnComplete Complete;
+	FJobProcessingStateChanged JobProcessingStateChanged;
+
+	std::atomic<bool> bIsRunning;
+	std::atomic<bool> bStopRequested;
+	TUniquePtr<FRunnableThread> Thread;
+};
+
+}
diff --git a/Engine/Plugins/VirtualProduction/CaptureManager/CaptureManagerApp/Source/CaptureManagerEditor/Private/IngestManagement/IngestJobSettingsManager.cpp b/Engine/Plugins/VirtualProduction/CaptureManager/CaptureManagerApp/Source/CaptureManagerEditor/Private/IngestManagement/IngestJobSettingsManager.cpp
new file mode 100644
--- /dev/null
+++ b/Engine/Plugins/VirtualProduction/CaptureManager/CaptureManagerApp/Source/CaptureManagerEditor/Private/IngestManagement/IngestJobSettingsManager.cpp
@@ -0,0 +1,97 @@
+// Copyright Epic Games, Inc. All Rights Reserved.
+
+#include "IngestJobSettingsManager.h"
+
+namespace UE::CaptureManager
+{
+
+FIngestJobSettingsManager::FIngestJobSettingsManager()
+{
+	SelectedPreset = Presets.Emplace_GetRef(NewObject<UDefaultJobPreset>()).Get();
+}
+
+TArray<FText> FIngestJobSettingsManager::GetPresetDisplayNames() const
+{
+	FScopeLock Lock(&CriticalSection);
+
+	TArray<FText> Names;
+	Names.Reserve(Presets.Num());
+
+	for (const TStrongObjectPtr<const UIngestJobSettingsBase>& Preset : Presets)
+	{
+		Names.Emplace(Preset->DisplayName);
+	}
+
+	return Names;
+}
+
+bool FIngestJobSettingsManager::SelectPresetByDisplayName(const FText& DisplayName)
+{
+	FScopeLock Lock(&CriticalSection);
+
+	const TStrongObjectPtr<const UIngestJobSettingsBase>* Preset = Presets.FindByPredicate(
+		[&DisplayName](const TStrongObjectPtr<const UIngestJobSettingsBase>& InPreset)
+		{
+			return InPreset && InPreset->DisplayName.CompareTo(DisplayName) == 0;
+		}
+	);
+
+	if (Preset)
+	{
+		SelectedPreset = Preset->Get();
+		return true;
+	}
+
+	return false;
+}
+
+const UIngestJobSettingsBase& FIngestJobSettingsManager::GetSelectedPreset() const
+{
+	FScopeLock Lock(&CriticalSection);
+	check(SelectedPreset);
+	return *SelectedPreset;
+}
+
+TObjectPtr<UIngestJobSettingsBase> FIngestJobSettingsManager::CreateSettingsFromSelectedPreset(const FGuid& JobGuid)
+{
+	FScopeLock Lock(&CriticalSection);
+
+	TStrongObjectPtr<UIngestJobSettingsBase>& Object = Settings.Emplace_GetRef(DuplicateObject(SelectedPreset.Get(), nullptr));
+	Object->JobGuid = JobGuid;
+	return Object.Get();
+}
+
+int32 FIngestJobSettingsManager::RemoveSettings(const TArray<FGuid>& JobGuids)
+{
+	FScopeLock Lock(&CriticalSection);
+
+	const int32 NumRemoved = Settings.RemoveAll(
+		[&JobGuids](const TStrongObjectPtr<UIngestJobSettingsBase>& Entry)
+		{
+			return JobGuids.Contains(Entry->JobGuid);
+		}
+	);
+
+	return NumRemoved;
+}
+
+TObjectPtr<UIngestJobSettingsBase> FIngestJobSettingsManager::GetSettings(const FGuid& JobGuid) const
+{
+	FScopeLock Lock(&CriticalSection);
+
+	const TStrongObjectPtr<UIngestJobSettingsBase>* Found = Settings.FindByPredicate(
+		[&JobGuid](const TStrongObjectPtr<const UIngestJobSettingsBase>& Entry)
+		{
+			return Entry->JobGuid == JobGuid;
+		}
+	);
+
+	if (Found)
+	{
+		return Found->Get();
+	}
+
+	return nullptr;
+}
+
+} // namespace UE::CaptureManager
\ No newline at end of file
diff --git a/Engine/Plugins/VirtualProduction/CaptureManager/CaptureManagerApp/Source/CaptureManagerEditor/Private/IngestManagement/IngestJobSettingsManager.h b/Engine/Plugins/VirtualProduction/CaptureManager/CaptureManagerApp/Source/CaptureManagerEditor/Private/IngestManagement/IngestJobSettingsManager.h
new file mode 100644
--- /dev/null
+++ b/Engine/Plugins/VirtualProduction/CaptureManager/CaptureManagerApp/Source/CaptureManagerEditor/Private/IngestManagement/IngestJobSettingsManager.h
@@ -0,0 +1,30 @@
+// Copyright Epic Games, Inc. All Rights Reserved.
+
+#pragma once
+
+#include "UIngestJobSettings.h"
+
+namespace UE::CaptureManager
+{
+
+class FIngestJobSettingsManager
+{
+public:
+	FIngestJobSettingsManager();
+
+	TArray<FText> GetPresetDisplayNames() const;
+	bool SelectPresetByDisplayName(const FText& DisplayName);
+	const UIngestJobSettingsBase& GetSelectedPreset() const;
+
+	TObjectPtr<UIngestJobSettingsBase> CreateSettingsFromSelectedPreset(const FGuid& JobGuid);
+	TObjectPtr<UIngestJobSettingsBase> GetSettings(const FGuid& JobGuid) const;
+	int32 RemoveSettings(const TArray<FGuid>& JobGuids);
+
+private:
+	mutable FCriticalSection CriticalSection;
+	TObjectPtr<const UIngestJobSettingsBase> SelectedPreset;
+	TArray<TStrongObjectPtr<const UIngestJobSettingsBase>> Presets;
+	TArray<TStrongObjectPtr<UIngestJobSettingsBase>> Settings;
+};
+
+}
\ No newline at end of file
diff --git a/Engine/Plugins/VirtualProduction/CaptureManager/CaptureManagerApp/Source/CaptureManagerEditor/Private/IngestManagement/UIngestJobSettings.cpp b/Engine/Plugins/VirtualProduction/CaptureManager/CaptureManagerApp/Source/CaptureManagerEditor/Private/IngestManagement/UIngestJobSettings.cpp
new file mode 100644
--- /dev/null
+++ b/Engine/Plugins/VirtualProduction/CaptureManager/CaptureManagerApp/Source/CaptureManagerEditor/Private/IngestManagement/UIngestJobSettings.cpp
@@ -0,0 +1,38 @@
+// Copyright Epic Games, Inc. All Rights Reserved.
+
+#include "UIngestJobSettings.h"
+
+#include "Misc/Paths.h"
+
+#include "Settings/CaptureManagerSettings.h"
+
+#define LOCTEXT_NAMESPACE "UIngestJobSettings"
+
+UIngestJobSettingsBase::UIngestJobSettingsBase(const FObjectInitializer& ObjectInitializer)
+	: Super(ObjectInitializer)
+{
+}
+
+UDefaultJobPreset::UDefaultJobPreset(const FObjectInitializer& ObjectInitializer)
+	: Super(ObjectInitializer)
+{
+	DisplayName = LOCTEXT("DefaultText", "Default");
+	LocalWorkFolder = GetDefault<UCaptureManagerSettings>()->DefaultConversionDirectory;
+	ImageFormat = EOutputImageFormat::JPEG;
+	ImageFileNamePrefix = TEXT("frame");
+	ImagePixelFormat = EImagePixelFormat::U8_BGRA;
+	AudioFormat = EAudioFormat::WAV;
+	AudioFileNamePrefix = TEXT("audio");
+}
+
+void UDefaultJobPreset::PostEditChangeProperty(struct FPropertyChangedEvent& InPropertyChangedEvent)
+{
+	Super::PostEditChangeProperty(InPropertyChangedEvent);
+
+	if (LocalWorkFolder.Path.IsEmpty())
+	{
+		LocalWorkFolder = GetDefault<UCaptureManagerSettings>()->DefaultConversionDirectory;
+	}
+}
+
+#undef LOCTEXT_NAMESPACE
diff --git a/Engine/Plugins/VirtualProduction/CaptureManager/CaptureManagerApp/Source/CaptureManagerEditor/Private/IngestManagement/UIngestJobSettings.h b/Engine/Plugins/VirtualProduction/CaptureManager/CaptureManagerApp/Source/CaptureManagerEditor/Private/IngestManagement/UIngestJobSettings.h
new file mode 100644
--- /dev/null
+++ b/Engine/Plugins/VirtualProduction/CaptureManager/CaptureManagerApp/Source/CaptureManagerEditor/Private/IngestManagement/UIngestJobSettings.h
@@ -0,0 +1,76 @@
+// Copyright Epic Games, Inc. All Rights Reserved.
+
+#pragma once
+
+#include "UIngestJobSettings.generated.h"
+
+// In order to enable a rather nice "reset to preset default" functionality in the job details view, we specify presets
+// by inheriting from the base object and overriding the default property values in the constructor. If we simply 
+// assign these values (post-construction) then these values are not treated as "default" and this functionality 
+// doesn't work anymore.
+
+UENUM()
+enum class EOutputImageFormat : uint8
+{
+	JPEG = 0			UMETA(DisplayName = "jpg"),
+	PNG					UMETA(DisplayName = "png"),
+};
+
+UENUM()
+enum class EImagePixelFormat : uint8
+{
+	U8_BGRA = 0			UMETA(DisplayName = "U8 BGRA"),
+	U8_Mono				UMETA(DisplayName = "U8 Mono")
+};
+
+UENUM()
+enum class EAudioFormat : uint8
+{
+	WAV = 0				UMETA(DisplayName = "wav")
+};
+
+UCLASS(Abstract)
+class UIngestJobSettingsBase : public UObject
+{
+	GENERATED_BODY()
+
+public:
+	UPROPERTY(AdvancedDisplay)
+	FGuid JobGuid;
+
+	UPROPERTY(AdvancedDisplay)
+	FText DisplayName;
+
+	UPROPERTY(Category = Conversion, EditAnywhere)
+	FDirectoryPath LocalWorkFolder;
+
+	UPROPERTY(Category = Video, EditAnywhere)
+	EOutputImageFormat ImageFormat;
+
+	UPROPERTY(Category = Video, EditAnywhere, DisplayName = "File Name")
+	FString ImageFileNamePrefix;
+
+	UPROPERTY(Category = Video, EditAnywhere)
+	EImagePixelFormat ImagePixelFormat;
+
+	UPROPERTY(Category = Audio, EditAnywhere)
+	EAudioFormat AudioFormat;
+
+	UPROPERTY(Category = Audio, EditAnywhere, DisplayName = "File Name")
+	FString AudioFileNamePrefix;
+
+	UIngestJobSettingsBase(const FObjectInitializer& ObjectInitializer);
+};
+
+UCLASS()
+class UDefaultJobPreset : public UIngestJobSettingsBase
+{
+	GENERATED_BODY()
+
+public:
+	UDefaultJobPreset(const FObjectInitializer& ObjectInitializer);
+
+private:
+
+	virtual void PostEditChangeProperty(struct FPropertyChangedEvent& InPropertyChangedEvent) override;
+};
diff --git a/Engine/Plugins/VirtualProduction/CaptureManager/CaptureManagerApp/Source/CaptureManagerEditor/Private/Scripting/LiveLinkHubScriptStatus.cpp b/Engine/Plugins/VirtualProduction/CaptureManager/CaptureManagerApp/Source/CaptureManagerEditor/Private/Scripting/LiveLinkHubScriptStatus.cpp
new file mode 100644
--- /dev/null
+++ b/Engine/Plugins/VirtualProduction/CaptureManager/CaptureManagerApp/Source/CaptureManagerEditor/Private/Scripting/LiveLinkHubScriptStatus.cpp
@@ -0,0 +1,29 @@
+// Copyright Epic Games, Inc. All Rights Reserved.
+
+#include "LiveLinkHubScriptStatus.h"
+
+bool FLiveLinkHubScriptStatus::IsValid() const
+{
+	// Since we can't force the status object to be constructed with a valid state (UObjects require default 
+	// constructor) we use this function to verify it.
+	return (bSuccess && ErrorMessage.IsEmpty()) || (!bSuccess && !ErrorMessage.IsEmpty());
+}
+
+FLiveLinkHubScriptStatus FLiveLinkHubScriptStatus::Error(FText InErrorMessage)
+{
+	FLiveLinkHubScriptStatus Status;
+	Status.bSuccess = false;
+	Status.ErrorMessage = MoveTemp(InErrorMessage);
+
+	check(Status.IsValid());
+	return Status;
+}
+
+FLiveLinkHubScriptStatus FLiveLinkHubScriptStatus::Success()
+{
+	FLiveLinkHubScriptStatus Status;
+	Status.bSuccess = true;
+
+	check(Status.IsValid());
+	return Status;
+}
diff --git a/Engine/Plugins/VirtualProduction/CaptureManager/CaptureManagerApp/Source/CaptureManagerEditor/Private/Scripting/LiveLinkHubScriptStatus.h b/Engine/Plugins/VirtualProduction/CaptureManager/CaptureManagerApp/Source/CaptureManagerEditor/Private/Scripting/LiveLinkHubScriptStatus.h
new file mode 100644
--- /dev/null
+++ b/Engine/Plugins/VirtualProduction/CaptureManager/CaptureManagerApp/Source/CaptureManagerEditor/Private/Scripting/LiveLinkHubScriptStatus.h
@@ -0,0 +1,22 @@
+// Copyright Epic Games, Inc. All Rights Reserved.
+
+#pragma once
+
+#include "LiveLinkHubScriptStatus.generated.h"
+
+USTRUCT(BlueprintType)
+struct FLiveLinkHubScriptStatus
+{
+	GENERATED_BODY()
+
+	static FLiveLinkHubScriptStatus Success();
+	static FLiveLinkHubScriptStatus Error(FText InErrorMessage);
+
+	bool IsValid() const;
+
+	UPROPERTY(VisibleAnywhere, BlueprintReadOnly, Category = Status)
+	bool bSuccess = false;
+
+	UPROPERTY(VisibleAnywhere, BlueprintReadOnly, Category = Status)
+	FText ErrorMessage;
+};
diff --git a/Engine/Plugins/VirtualProduction/CaptureManager/CaptureManagerApp/Source/CaptureManagerEditor/Private/Scripting/ULiveLinkHubArchiveCaptureDevice.cpp b/Engine/Plugins/VirtualProduction/CaptureManager/CaptureManagerApp/Source/CaptureManagerEditor/Private/Scripting/ULiveLinkHubArchiveCaptureDevice.cpp
new file mode 100644
--- /dev/null
+++ b/Engine/Plugins/VirtualProduction/CaptureManager/CaptureManagerApp/Source/CaptureManagerEditor/Private/Scripting/ULiveLinkHubArchiveCaptureDevice.cpp
@@ -0,0 +1,12 @@
+// Copyright Epic Games, Inc. All Rights Reserved.
+
+#include "ULiveLinkHubArchiveCaptureDevice.h"
+
+TMap<FString, FPropertyValue> ULiveLinkHubArchiveCaptureDevice::GetCreationParameters()
+{
+	check(!ArchivePath.IsEmpty());
+
+	TMap<FString, FPropertyValue> ParamValues;
+	ParamValues.Emplace(TEXT("StoragePath"), MakePropertyValue<FString>(ArchivePath));
+	return ParamValues;
+}
diff --git a/Engine/Plugins/VirtualProduction/CaptureManager/CaptureManagerApp/Source/CaptureManagerEditor/Private/Scripting/ULiveLinkHubArchiveCaptureDevice.h b/Engine/Plugins/VirtualProduction/CaptureManager/CaptureManagerApp/Source/CaptureManagerEditor/Private/Scripting/ULiveLinkHubArchiveCaptureDevice.h
new file mode 100644
--- /dev/null
+++ b/Engine/Plugins/VirtualProduction/CaptureManager/CaptureManagerApp/Source/CaptureManagerEditor/Private/Scripting/ULiveLinkHubArchiveCaptureDevice.h
@@ -0,0 +1,22 @@
+// Copyright Epic Games, Inc. All Rights Reserved.
+
+#pragma once
+
+#include "ULiveLinkHubCaptureDevice.h"
+
+#include "ULiveLinkHubArchiveCaptureDevice.generated.h"
+
+
+UCLASS(BlueprintType)
+class CAPTUREMANAGEREDITOR_API ULiveLinkHubArchiveCaptureDevice : public ULiveLinkHubCaptureDevice
+{
+	GENERATED_BODY()
+
+public:
+	UPROPERTY(EditAnywhere, BlueprintReadWrite, Category = Location)
+	FString ArchivePath;
+
+private:
+	virtual TMap<FString, FPropertyValue> GetCreationParameters() override;
+};
+
diff --git a/Engine/Plugins/VirtualProduction/CaptureManager/CaptureManagerApp/Source/CaptureManagerEditor/Private/Scripting/ULiveLinkHubCaptureDevice.cpp b/Engine/Plugins/VirtualProduction/CaptureManager/CaptureManagerApp/Source/CaptureManagerEditor/Private/Scripting/ULiveLinkHubCaptureDevice.cpp
new file mode 100644
--- /dev/null
+++ b/Engine/Plugins/VirtualProduction/CaptureManager/CaptureManagerApp/Source/CaptureManagerEditor/Private/Scripting/ULiveLinkHubCaptureDevice.cpp
@@ -0,0 +1,367 @@
+// Copyright Epic Games, Inc. All Rights Reserved.
+
+#include "ULiveLinkHubCaptureDevice.h"
+
+#include "CaptureSourceFrameworkModule.h"
+
+#include "Capabilities/IngestCapability.h"
+
+#define LOCTEXT_NAMESPACE "LiveLinkHubScriptedIngest"
+
+static const FText NoSessionText = LOCTEXT("NoSession", "No session");
+static const FText StartTimedOutText = LOCTEXT("StartTimedOut", "Timed out waiting for device to start");
+static const FText InvalidTakeSessionIdText = LOCTEXT("InvalidTakeSessionId", "Session ID for this take is incorrect, you can not refer to a take from previous device session");
+static const FText IngestNotSupportedText = LOCTEXT("IngestNotSupported", "Device does not support ingest");
+
+static const FText FailedToCreateCaptureSourceFormatText = LOCTEXT("FailedToCreateCaptureSourceFormat", "Failed to create capture source: {0}");
+static const FText ImportFailedFormatText = LOCTEXT("ImportFailedFormat", "Import failed: {0}");
+
+struct ULiveLinkHubCaptureDevice::FImpl
+{
+	struct FSession
+	{
+		explicit FSession(UE::CaptureManager::FCaptureSourceId InCaptureSourceId);
+
+		// This session ID is used to verify that any take objects we have returned to the user still belong to the 
+		// current start/stop session. This should prevent issues where take IDs have changed after a start/stop call 
+		// on the device. As an example consider an archive capture source, it may register takes in a different order
+		// depending on a non-deterministic traversal of the filesystem and so the same take from an earlier session
+		// may now have a different take ID.
+		const FGuid Id;
+
+		const UE::CaptureManager::FCaptureSourceId CaptureSourceId;
+	};
+
+	static UE::CaptureManager::FCaptureSourceManager& GetCaptureSourceManager();
+
+	FImpl();
+
+	FLiveLinkHubScriptStatus ImportTake(const FLiveLinkHubTakeMetadata& InTake, const FLiveLinkHubConversionSettings& InConversionSettings) const;
+	FLiveLinkHubScriptStatus FetchTakes(TArray<FLiveLinkHubTakeMetadata>& OutTakes) const;
+
+	FLiveLinkHubScriptStatus Start(const int32 InTimeoutSeconds, TMap<FString, FPropertyValue> InParamValues, const FString& InFactoryId, const FString& InName);
+	FLiveLinkHubScriptStatus Stop();
+
+	bool StartCaptureSource(UE::CaptureManager::FCaptureSourceId InCaptureSourceId, int32 InTimeoutSeconds);
+	void StopCaptureSource(UE::CaptureManager::FCaptureSourceId InCaptureSourceId);
+
+	TOptional<FSession> Session;
+	UE::CaptureManager::FCaptureSourceManager& CaptureSourceManager;
+};
+
+ULiveLinkHubCaptureDevice::FImpl::FSession::FSession(const UE::CaptureManager::FCaptureSourceId InCaptureSourceId) :
+	Id(FGuid::NewGuid()),
+	CaptureSourceId(InCaptureSourceId)
+{
+}
+
+UE::CaptureManager::FCaptureSourceManager& ULiveLinkHubCaptureDevice::FImpl::GetCaptureSourceManager()
+{
+	using namespace  UE::CaptureManager;
+
+	FCaptureSourceFrameworkModule& CaptureSourceModule = FModuleManager::LoadModuleChecked<FCaptureSourceFrameworkModule>("CaptureSourceFramework");
+	return CaptureSourceModule.GetCaptureSourceManager();
+}
+
+ULiveLinkHubCaptureDevice::FImpl::FImpl() :
+	CaptureSourceManager(GetCaptureSourceManager())
+{
+}
+
+FLiveLinkHubScriptStatus ULiveLinkHubCaptureDevice::FImpl::Start(const int32 InTimeoutSeconds, TMap<FString, FPropertyValue> InParamValues, const FString& InFactoryId, const FString& InName)
+{
+	using namespace  UE::CaptureManager;
+
+	if (Session)
+	{
+		return FLiveLinkHubScriptStatus::Success();
+	}
+
+	FCaptureSourceManager::FCaptureSourceResult CaptureSourceId = CaptureSourceManager.CreateCaptureSource(InFactoryId, InName, MoveTemp(InParamValues));
+
+	if (CaptureSourceId.HasError())
+	{
+		return FLiveLinkHubScriptStatus::Error(FText::Format(FailedToCreateCaptureSourceFormatText, FText::FromString(CaptureSourceId.GetError().GetMessage())));
+	}
+
+	const bool bIsStarted = StartCaptureSource(CaptureSourceId.GetValue(), InTimeoutSeconds);
+
+	if (bIsStarted)
+	{
+		Session = FSession(CaptureSourceId.StealValue());
+		return FLiveLinkHubScriptStatus::Success();
+	}
+	else
+	{
+		return FLiveLinkHubScriptStatus::Error(StartTimedOutText);
+	}
+}
+
+bool ULiveLinkHubCaptureDevice::FImpl::StartCaptureSource(const UE::CaptureManager::FCaptureSourceId InCaptureSourceId, const int32 InTimeoutSeconds)
+{
+	using namespace UE;
+	using namespace UE::CaptureManager;
+
+	FCaptureSourceHandle CaptureSourceHandle = CaptureSourceManager.GetCaptureSource(InCaptureSourceId);
+	FEventRef ReachableEvent;
+
+	TManagedDelegate<TSharedPtr<const FCaptureEvent>> OnReachableEvent(
+		[&ReachableEvent](const TSharedPtr<const FCaptureEvent> CaptureEvent)
+		{
+			if (CaptureEvent)
+			{
+				TSharedPtr<const FReachableEvent> Event = StaticCastSharedPtr<const FReachableEvent>(CaptureEvent);
+				if (Event->bIsReachable)
+				{
+					ReachableEvent->Trigger();
+				}
+			}
+		},
+		EDelegateExecutionThread::InternalThread
+	);
+
+	CaptureSourceHandle.SubscribeToEvent(FReachableEvent::Name, MoveTemp(OnReachableEvent));
+
+	ON_SCOPE_EXIT
+	{
+		CaptureSourceHandle.UnsubscribeAll();
+	};
+
+	AsyncTask(
+		ENamedThreads::AnyThread,
+		[&CaptureSourceHandle]()
+		{
+			CaptureSourceHandle.Start();
+		}
+	);
+
+	return ReachableEvent->Wait(FTimespan::FromSeconds(InTimeoutSeconds));
+}
+
+FLiveLinkHubScriptStatus ULiveLinkHubCaptureDevice::FImpl::Stop()
+{
+	using namespace  UE::CaptureManager;
+
+	if (!Session)
+	{
+		return FLiveLinkHubScriptStatus::Success();
+	}
+
+	StopCaptureSource(Session->CaptureSourceId);
+	CaptureSourceManager.RemoveCaptureSource(Session->CaptureSourceId);
+	Session.Reset();
+
+	return FLiveLinkHubScriptStatus::Success();
+}
+
+void ULiveLinkHubCaptureDevice::FImpl::StopCaptureSource(const UE::CaptureManager::FCaptureSourceId InCaptureSourceId)
+{
+	using namespace  UE::CaptureManager;
+
+	FCaptureSourceHandle CaptureSourceHandle = CaptureSourceManager.GetCaptureSource(InCaptureSourceId);
+	FEventRef UnreachableEvent;
+
+	TManagedDelegate<TSharedPtr<const FCaptureEvent>> OnReachableEvent(
+		[&UnreachableEvent](const TSharedPtr<const FCaptureEvent> CaptureEvent)
+		{
+			if (CaptureEvent)
+			{
+				TSharedPtr<const FReachableEvent> Event = StaticCastSharedPtr<const FReachableEvent>(CaptureEvent);
+				if (!Event->bIsReachable)
+				{
+					UnreachableEvent->Trigger();
+					return;
+				}
+			}
+		},
+		EDelegateExecutionThread::InternalThread
+	);
+
+	CaptureSourceHandle.SubscribeToEvent(FReachableEvent::Name, MoveTemp(OnReachableEvent));
+
+	ON_SCOPE_EXIT
+	{
+		CaptureSourceHandle.UnsubscribeAll();
+	};
+
+	AsyncTask(
+		ENamedThreads::AnyThread,
+		[&CaptureSourceHandle]()
+		{
+			CaptureSourceHandle.Stop();
+		}
+	);
+
+	UnreachableEvent->Wait();
+}
+
+FLiveLinkHubScriptStatus ULiveLinkHubCaptureDevice::FImpl::ImportTake(const FLiveLinkHubTakeMetadata& InTake, const FLiveLinkHubConversionSettings& InConversionSettings) const
+{
+	using namespace  UE::CaptureManager;
+
+	if (!Session)
+	{
+		return FLiveLinkHubScriptStatus::Error(NoSessionText);
+	}
+
+	if (InTake.SessionId != Session->Id)
+	{
+		return FLiveLinkHubScriptStatus::Error(InvalidTakeSessionIdText);
+	}
+
+	FCaptureSourceHandle CaptureSourceHandle = CaptureSourceManager.GetCaptureSource(Session->CaptureSourceId);
+	FIngestCapability* IngestCapability = CaptureSourceHandle.GetCapability<FIngestCapability>(FIngestCapability::Name);
+
+	if (!IngestCapability)
+	{
+		return FLiveLinkHubScriptStatus::Error(IngestNotSupportedText);
+	}
+
+	FEventRef ImportCompleteEvent;
+	bool bImportSuccess = false;
+	FText ErrorMessage;
+
+	TManagedDelegate<FTakeId, TValueOrError<void, FIngestCapabilityError>> OnImportFinished(
+		[&ImportCompleteEvent, &bImportSuccess, &ErrorMessage](FTakeId InTakeId, TValueOrError<void, FIngestCapabilityError> InImportResult)
+		{
+			if (InImportResult.HasValue())
+			{
+				bImportSuccess = true;
+			}
+			else
+			{
+				ErrorMessage = FText::Format(ImportFailedFormatText, FText::FromString(InImportResult.GetError().GetMessage()));
+			}
+
+			ImportCompleteEvent->Trigger();
+		},
+		EDelegateExecutionThread::InternalThread
+	);
+
+	FIngestOptions::FAudio AudioOptions =
+	{
+		.FileNamePrefix = InConversionSettings.AudioFileNamePrefix,
+		.Format = InConversionSettings.AudioFormat
+	};
+
+	FIngestOptions::FVideo VideoOptions = 
+	{
+		.FileNamePrefix = InConversionSettings.VideoFileNamePrefix,
+		.Format = InConversionSettings.VideoFormat,
+		.PixelFormat = UE::CaptureManager::EImagePixelFormat::Default
+	};
+
+	FIngestOptions Options =
+	{
+		.Location = InConversionSettings.WorkingDirectory,
+		.Video = VideoOptions,
+		.Audio = AudioOptions
+	};
+
+	IngestCapability->ImportTakeAsync(InTake.TakeId, MoveTemp(Options), OnImportFinished, TManagedDelegate<FTakeId, double>());
+	ImportCompleteEvent->Wait();
+
+	if (bImportSuccess)
+	{
+		return FLiveLinkHubScriptStatus::Success();
+	}
+
+	return FLiveLinkHubScriptStatus::Error(ErrorMessage);
+}
+
+FLiveLinkHubScriptStatus ULiveLinkHubCaptureDevice::FImpl::FetchTakes(TArray<FLiveLinkHubTakeMetadata>& OutTakes) const
+{
+	using namespace  UE::CaptureManager;
+
+	if (!OutTakes.IsEmpty())
+	{
+		// Make sure the output array is in a valid state, in case anything goes wrong below
+		OutTakes.Empty();
+	}
+
+	if (!Session)
+	{
+		return FLiveLinkHubScriptStatus::Error(NoSessionText);
+	}
+
+	FCaptureSourceHandle CaptureSourceHandle = CaptureSourceManager.GetCaptureSource(Session->CaptureSourceId);
+	FIngestCapability* IngestCapability = CaptureSourceHandle.GetCapability<FIngestCapability>(FIngestCapability::Name);
+
+	if (!IngestCapability)
+	{
+		return FLiveLinkHubScriptStatus::Error(IngestNotSupportedText);
+	}
+
+	FEventRef TakeListUpdatedEvent;
+	TArray<FLiveLinkHubTakeMetadata> Takes;
+
+	TManagedDelegate<TArray<FTakeId>> OnTakesListUpdated(
+		[SessionId = Session->Id, IngestCapability, &Takes, &TakeListUpdatedEvent](const TArray<FTakeId>& InTakeIds)
+		{
+			for (const FTakeId& TakeId : InTakeIds)
+			{
+				TOptional<FTakeMetadata> TakeMetadataOpt = IngestCapability->GetTakeMetadata(TakeId);
+
+				if (TakeMetadataOpt.IsSet())
+				{
+					const FTakeMetadata TakeMetadata = MoveTemp(TakeMetadataOpt.GetValue());
+
+					FLiveLinkHubTakeMetadata Take;
+					Take.SlateName = TakeMetadata.Slate;
+					Take.TakeNumber = TakeMetadata.TakeNumber;
+					Take.TakeId = TakeId;
+					Take.SessionId = SessionId;
+					Take.DateTime = TakeMetadata.DateTime.Get(FDateTime()).ToIso8601();
+
+					Takes.Emplace(MoveTemp(Take));
+				}
+			}
+
+			TakeListUpdatedEvent->Trigger();
+		},
+		EDelegateExecutionThread::InternalThread
+	);
+
+	IngestCapability->UpdateTakeListAsync(MoveTemp(OnTakesListUpdated));
+	TakeListUpdatedEvent->Wait();
+
+	OutTakes = MoveTemp(Takes);
+	return FLiveLinkHubScriptStatus::Success();
+}
+
+ULiveLinkHubCaptureDevice::ULiveLinkHubCaptureDevice() :
+	ImplPtr(MakePimpl<FImpl>())
+{
+}
+
+ULiveLinkHubCaptureDevice::~ULiveLinkHubCaptureDevice() = default;
+
+FLiveLinkHubScriptStatus ULiveLinkHubCaptureDevice::Start(const int32 InTimeoutSeconds)
+{
+	TMap<FString, FPropertyValue> ParamValues = GetCreationParameters();
+	return ImplPtr->Start(InTimeoutSeconds, MoveTemp(ParamValues), FactoryId, Name);
+}
+
+FLiveLinkHubScriptStatus ULiveLinkHubCaptureDevice::Stop()
+{
+	return ImplPtr->Stop();
+}
+
+FLiveLinkHubScriptStatus ULiveLinkHubCaptureDevice::ImportTake(const FLiveLinkHubTakeMetadata& InTake, const FLiveLinkHubConversionSettings& InConversionSettings) const
+{
+	return ImplPtr->ImportTake(InTake, InConversionSettings);
+}
+
+FLiveLinkHubFetchTakesResult ULiveLinkHubCaptureDevice::FetchTakes() const
+{
+	TArray<FLiveLinkHubTakeMetadata> Takes;
+	FLiveLinkHubScriptStatus Status = ImplPtr->FetchTakes(Takes);
+
+	FLiveLinkHubFetchTakesResult Result;
+	Result.Status = MoveTemp(Status);
+	Result.Takes = MoveTemp(Takes);
+
+	return Result;
+}
+
+#undef LOCTEXT_NAMESPACE
diff --git a/Engine/Plugins/VirtualProduction/CaptureManager/CaptureManagerApp/Source/CaptureManagerEditor/Private/Scripting/ULiveLinkHubCaptureDevice.h b/Engine/Plugins/VirtualProduction/CaptureManager/CaptureManagerApp/Source/CaptureManagerEditor/Private/Scripting/ULiveLinkHubCaptureDevice.h
new file mode 100644
--- /dev/null
+++ b/Engine/Plugins/VirtualProduction/CaptureManager/CaptureManagerApp/Source/CaptureManagerEditor/Private/Scripting/ULiveLinkHubCaptureDevice.h
@@ -0,0 +1,102 @@
+// Copyright Epic Games, Inc. All Rights Reserved.
+
+#pragma once
+
+#include "CaptureSourceCapabilityProperty.h"
+#include "LiveLinkHubScriptStatus.h"
+
+#include "ULiveLinkHubCaptureDevice.generated.h"
+
+USTRUCT(BlueprintType)
+struct FLiveLinkHubTakeMetadata
+{
+	GENERATED_BODY()
+
+	UPROPERTY(EditAnywhere, BlueprintReadWrite, Category = Metadata)
+	FString SlateName;
+
+	UPROPERTY(EditAnywhere, BlueprintReadWrite, Category = Metadata)
+	int32 TakeNumber = 1;
+
+	UPROPERTY(EditAnywhere, BlueprintReadWrite, Category = Metadata)
+	FString DateTime;
+
+	UPROPERTY()
+	int32 TakeId = 0;
+
+	UPROPERTY()
+	/** Identifies which start/stop session this information belongs to, to ensure the TakeId is still valid */
+	FGuid SessionId;
+};
+
+USTRUCT(BlueprintType)
+struct FLiveLinkHubFetchTakesResult
+{
+	GENERATED_BODY()
+
+	UPROPERTY(VisibleAnywhere, BlueprintReadOnly, Category = Result)
+	FLiveLinkHubScriptStatus Status;
+
+	UPROPERTY(VisibleAnywhere, BlueprintReadOnly, Category = Result)
+	TArray<FLiveLinkHubTakeMetadata> Takes;
+};
+
+USTRUCT(BlueprintType)
+struct FLiveLinkHubConversionSettings
+{
+	GENERATED_BODY()
+
+	UPROPERTY(EditAnywhere, BlueprintReadWrite, Category = General)
+	FString WorkingDirectory;
+
+	UPROPERTY(EditAnywhere, BlueprintReadWrite, Category = Audio)
+	FString AudioFormat;
+
+	UPROPERTY(EditAnywhere, BlueprintReadWrite, Category = Audio)
+	FString AudioFileNamePrefix;
+
+	UPROPERTY(EditAnywhere, BlueprintReadWrite, Category = Video)
+	FString VideoFormat;
+
+	UPROPERTY(EditAnywhere, BlueprintReadWrite, Category = Video)
+	FString VideoFileNamePrefix;
+
+	// TODO: Make enum visible to python, or map python enum to c++ enum
+	// UPROPERTY(EditAnywhere, BlueprintReadWrite, Category = Video)
+	// int32 PixelFormat = 0;
+};
+
+UCLASS(Abstract)
+class ULiveLinkHubCaptureDevice : public UObject
+{
+	GENERATED_BODY()
+
+public:
+	ULiveLinkHubCaptureDevice();
+	~ULiveLinkHubCaptureDevice();
+
+	UFUNCTION(BlueprintCallable, Category = Control)
+	FLiveLinkHubScriptStatus Start(int32 InTimeoutSeconds);
+
+	UFUNCTION(BlueprintCallable, Category = Control)
+	FLiveLinkHubScriptStatus Stop();
+
+	UFUNCTION(BlueprintCallable, Category = Control)
+	FLiveLinkHubScriptStatus ImportTake(const FLiveLinkHubTakeMetadata& InTake, const FLiveLinkHubConversionSettings& InConversionSettings) const;
+
+	UFUNCTION(BlueprintCallable, Category = Control)
+	FLiveLinkHubFetchTakesResult FetchTakes() const;
+
+	UPROPERTY(EditAnywhere, BlueprintReadWrite, Category = Info)
+	FString Name;
+
+	UPROPERTY(EditAnywhere, BlueprintReadWrite, Category = Info)
+	FString FactoryId;
+
+private:
+	struct FImpl;
+	TPimplPtr<FImpl> ImplPtr;
+
+	virtual TMap<FString, FPropertyValue> GetCreationParameters() PURE_VIRTUAL(ULiveLinkHubCaptureDevice::GetCreationParameters, return {}; );
+};
+
diff --git a/Engine/Plugins/VirtualProduction/CaptureManager/CaptureManagerApp/Source/CaptureManagerEditor/Private/Scripting/ULiveLinkHubImportClient.cpp b/Engine/Plugins/VirtualProduction/CaptureManager/CaptureManagerApp/Source/CaptureManagerEditor/Private/Scripting/ULiveLinkHubImportClient.cpp
new file mode 100644
--- /dev/null
+++ b/Engine/Plugins/VirtualProduction/CaptureManager/CaptureManagerApp/Source/CaptureManagerEditor/Private/Scripting/ULiveLinkHubImportClient.cpp
@@ -0,0 +1,218 @@
+// Copyright Epic Games, Inc. All Rights Reserved.
+
+#include "ULiveLinkHubImportClient.h"
+
+#include "LiveLinkHubDiscoveryModule.h"
+#include "CaptureCommunicationManagerModule.h"
+
+#include "Network/NetworkMisc.h"
+
+#define LOCTEXT_NAMESPACE "LiveLinkHubScriptedIngest"
+
+static const FText TimedOutWaitingForClientText = LOCTEXT("TimedOutWaitingForClient", "Timed out waiting for client");
+static const FText FailedToDetermineLocalIpAddressText = LOCTEXT("FailedToDetermineLocalIpAddress", "Failed to determine local IP address");
+
+struct ULiveLinkHubImportClient::FImpl
+{
+	static TSharedRef<UE::CaptureManager::FDiscoveryRequester> GetDiscoveryRequester();
+	static TSharedRef<FCommunicationManager> GetCommunicationManager();
+
+	FImpl();
+
+	FLiveLinkHubScriptStatus Start(int32 InTimeoutSeconds);
+	FLiveLinkHubScriptStatus Stop();
+
+	TValueOrError<UE::CaptureManager::FDiscoveredClient, FText> WaitForLocalClientDiscovery(double InTimeoutSeconds);
+	bool ConnectToClient(const UE::CaptureManager::FDiscoveredClient& InClient, double InTimeoutSeconds);
+	TOptional<UE::CaptureManager::FDiscoveredClient> FindExistingClientByIpAddress(const FString& InIpAddress);
+
+	TSharedRef<FCommunicationManager> CommunicationManager;
+};
+
+TSharedRef<UE::CaptureManager::FDiscoveryRequester> ULiveLinkHubImportClient::FImpl::GetDiscoveryRequester()
+{
+	FLiveLinkHubDiscovery& LiveLinkHubDiscovery =
+		FModuleManager::LoadModuleChecked<FLiveLinkHubDiscovery>("LiveLinkHubDiscovery");
+
+	return LiveLinkHubDiscovery.GetDiscoveryRequster();
+}
+
+TSharedRef<FCommunicationManager> ULiveLinkHubImportClient::FImpl::GetCommunicationManager()
+{
+	FCaptureCommunicationManagerModule& CaptureCommunicationManagerModule =
+		FModuleManager::LoadModuleChecked<FCaptureCommunicationManagerModule>("CaptureCommunicationManager");
+
+	return CaptureCommunicationManagerModule.GetManager();
+}
+
+ULiveLinkHubImportClient::FImpl::FImpl() :
+	CommunicationManager(GetCommunicationManager())
+{
+}
+
+FLiveLinkHubScriptStatus ULiveLinkHubImportClient::FImpl::Start(const int32 InTimeoutSeconds)
+{
+	using namespace UE::CaptureManager;
+
+	if (CommunicationManager->IsConnected())
+	{
+		return FLiveLinkHubScriptStatus::Success();
+	}
+
+	const double ExpiresAt = FPlatformTime::Seconds() + InTimeoutSeconds;
+
+	TValueOrError<FDiscoveredClient, FText> LocalClient = WaitForLocalClientDiscovery(InTimeoutSeconds);
+
+	if (LocalClient.HasError())
+	{
+		return FLiveLinkHubScriptStatus::Error(LocalClient.StealError());
+	}
+
+	const double ConnectionTimeout = ExpiresAt - FPlatformTime::Seconds();
+
+	if (ConnectionTimeout <= 0.0)
+	{
+		return FLiveLinkHubScriptStatus::Error(TimedOutWaitingForClientText);
+	}
+
+	const bool bIsConnected = ConnectToClient(LocalClient.GetValue(), ConnectionTimeout);
+
+	if (bIsConnected)
+	{
+		return FLiveLinkHubScriptStatus::Success();
+	}
+
+	return FLiveLinkHubScriptStatus::Error(TimedOutWaitingForClientText);
+}
+
+FLiveLinkHubScriptStatus ULiveLinkHubImportClient::FImpl::Stop()
+{
+	return FLiveLinkHubScriptStatus::Success();
+}
+
+TValueOrError<UE::CaptureManager::FDiscoveredClient, FText> ULiveLinkHubImportClient::FImpl::WaitForLocalClientDiscovery(
+	const double InTimeoutSeconds
+)
+{
+	using namespace UE::CaptureManager;
+
+	TOptional<FString> LocalIpAddress = GetLocalIpAddress();
+	check(LocalIpAddress);
+
+	if (!LocalIpAddress)
+	{
+		return MakeError(FailedToDetermineLocalIpAddressText);
+	}
+
+	FEventRef LocalClientDiscoveredEvent;
+	TOptional<FDiscoveredClient> LocalClient;
+
+	// Register a discovery listener BEFORE we query the already discovered clients - as it may fire while we are 
+	// processing that information.
+	FDelegateHandle ClientDiscoveredListener = GetDiscoveryRequester()->AddClientDiscoveredListener(
+		FDiscoveryRequester::FClientDiscovered::FDelegate::CreateLambda(
+			[&LocalClientDiscoveredEvent, &LocalIpAddress, &LocalClient](const FDiscoveredClient& Client)
+			{
+				if (Client.GetIPAddress() == *LocalIpAddress)
+				{
+					LocalClient = Client;
+					LocalClientDiscoveredEvent->Trigger();
+				}
+			}
+		)
+	);
+
+	ON_SCOPE_EXIT
+	{
+		GetDiscoveryRequester()->RemoveClientDiscoveredListener(ClientDiscoveredListener);
+	};
+
+	TOptional<FDiscoveredClient> ExistingLocalClient = FindExistingClientByIpAddress(*LocalIpAddress);
+
+	if (ExistingLocalClient)
+	{
+		return MakeValue(*ExistingLocalClient);
+	}
+
+	// Local client has not been discovered yet, so wait for it
+	LocalClientDiscoveredEvent->Wait(FTimespan::FromSeconds(InTimeoutSeconds));
+
+	if (LocalClient)
+	{
+		return MakeValue(*LocalClient);
+	}
+
+	return MakeError(TimedOutWaitingForClientText);
+}
+
+bool ULiveLinkHubImportClient::FImpl::ConnectToClient(
+	const UE::CaptureManager::FDiscoveredClient& InClient,
+	const double InTimeoutSeconds
+)
+{
+	FEventRef ClientConnectedEvent;
+
+	FDelegateHandle ClientConnectedDelegateHandle = CommunicationManager->AddConnectionStateListener(
+		FCommunicationManager::FConnectionStateChanged::FDelegate::CreateLambda(
+			[&ClientConnectedEvent](FMessageAddress MessageAddress, bool bIsConnected)
+			{
+				if (bIsConnected)
+				{
+					ClientConnectedEvent->Trigger();
+				}
+			}
+		)
+	);
+
+	CommunicationManager->Connect(InClient.GetMessageAddress(), InClient.GetIPAddress(), InClient.GetExportPort());
+
+	ON_SCOPE_EXIT
+	{
+		CommunicationManager->RemoveConnectionStateListener(ClientConnectedDelegateHandle);
+	};
+
+	return ClientConnectedEvent->Wait(FTimespan::FromSeconds(InTimeoutSeconds));
+}
+
+TOptional<UE::CaptureManager::FDiscoveredClient> ULiveLinkHubImportClient::FImpl::FindExistingClientByIpAddress(
+	const FString& InIpAddress
+)
+{
+	using namespace UE::CaptureManager;
+
+	TArray<FDiscoveredClient> DiscoveredClients = GetDiscoveryRequester()->GetDiscoveredClients();
+
+	FDiscoveredClient* ClientPtr = DiscoveredClients.FindByPredicate(
+		[&InIpAddress](const FDiscoveredClient& InClient)
+		{
+			return InClient.GetIPAddress() == InIpAddress;
+		}
+	);
+
+	if (ClientPtr)
+	{
+		FDiscoveredClient Client = MoveTemp(*ClientPtr);
+		return Client;
+	}
+
+	return {};
+}
+
+ULiveLinkHubImportClient::ULiveLinkHubImportClient() :
+	ImplPtr(MakePimpl<FImpl>())
+{
+}
+
+ULiveLinkHubImportClient::~ULiveLinkHubImportClient() = default;
+
+FLiveLinkHubScriptStatus ULiveLinkHubImportClient::Start(const int32 InTimeoutSeconds)
+{
+	return ImplPtr->Start(InTimeoutSeconds);
+}
+
+FLiveLinkHubScriptStatus ULiveLinkHubImportClient::Stop()
+{
+	return ImplPtr->Stop();
+}
+
+#undef LOCTEXT_NAMESPACE
diff --git a/Engine/Plugins/VirtualProduction/CaptureManager/CaptureManagerApp/Source/CaptureManagerEditor/Private/Scripting/ULiveLinkHubImportClient.h b/Engine/Plugins/VirtualProduction/CaptureManager/CaptureManagerApp/Source/CaptureManagerEditor/Private/Scripting/ULiveLinkHubImportClient.h
new file mode 100644
--- /dev/null
+++ b/Engine/Plugins/VirtualProduction/CaptureManager/CaptureManagerApp/Source/CaptureManagerEditor/Private/Scripting/ULiveLinkHubImportClient.h
@@ -0,0 +1,27 @@
+// Copyright Epic Games, Inc. All Rights Reserved.
+
+#pragma once
+
+#include "LiveLinkHubScriptStatus.h"
+
+#include "ULiveLinkHubImportClient.generated.h"
+
+UCLASS(BlueprintType)
+class ULiveLinkHubImportClient : public UObject
+{
+	GENERATED_BODY()
+	
+public:
+	ULiveLinkHubImportClient();
+	~ULiveLinkHubImportClient();
+
+	UFUNCTION(BlueprintCallable, Category = Control)
+	FLiveLinkHubScriptStatus Start(int32 InTimeoutSeconds);
+
+	UFUNCTION(BlueprintCallable, Category = Control)
+	FLiveLinkHubScriptStatus Stop();
+
+private:
+	struct FImpl;
+	TPimplPtr<FImpl> ImplPtr;
+};
diff --git a/Engine/Plugins/VirtualProduction/CaptureManager/CaptureManagerApp/Source/CaptureManagerEditor/Private/Scripting/ULiveLinkHubNetworkCaptureDevice.cpp b/Engine/Plugins/VirtualProduction/CaptureManager/CaptureManagerApp/Source/CaptureManagerEditor/Private/Scripting/ULiveLinkHubNetworkCaptureDevice.cpp
new file mode 100644
--- /dev/null
+++ b/Engine/Plugins/VirtualProduction/CaptureManager/CaptureManagerApp/Source/CaptureManagerEditor/Private/Scripting/ULiveLinkHubNetworkCaptureDevice.cpp
@@ -0,0 +1,15 @@
+// Copyright Epic Games, Inc. All Rights Reserved.
+
+#include "ULiveLinkHubNetworkCaptureDevice.h"
+
+TMap<FString, FPropertyValue> ULiveLinkHubNetworkCaptureDevice::GetCreationParameters()
+{
+	check(!IPAddress.IsEmpty());
+	check(Port != 0);
+
+	TMap<FString, FPropertyValue> ParamValues;
+	ParamValues.Emplace(TEXT("IpAddress"), MakePropertyValue<FString>(IPAddress));
+	ParamValues.Emplace(TEXT("Port"), MakePropertyValue<int64>(Port));
+	return ParamValues;
+}
+
