CaptureManager: Move plugin to Virtual Production directory

diff --git a/Engine/Plugins/VirtualProduction/CaptureManager/CaptureManagerApp/Source/CaptureManagerEditor/Private/Scripting/ULiveLinkHubNetworkCaptureDevice.h b/Engine/Plugins/VirtualProduction/CaptureManager/CaptureManagerApp/Source/CaptureManagerEditor/Private/Scripting/ULiveLinkHubNetworkCaptureDevice.h
new file mode 100644
--- /dev/null
+++ b/Engine/Plugins/VirtualProduction/CaptureManager/CaptureManagerApp/Source/CaptureManagerEditor/Private/Scripting/ULiveLinkHubNetworkCaptureDevice.h
@@ -0,0 +1,25 @@
+// Copyright Epic Games, Inc. All Rights Reserved.
+
+#pragma once
+
+#include "ULiveLinkHubCaptureDevice.h"
+
+#include "ULiveLinkHubNetworkCaptureDevice.generated.h"
+
+UCLASS(BlueprintType)
+class CAPTUREMANAGEREDITOR_API ULiveLinkHubNetworkCaptureDevice : public ULiveLinkHubCaptureDevice
+{
+	GENERATED_BODY()
+
+public:
+	UPROPERTY(EditAnywhere, BlueprintReadWrite, Category = Connection)
+	FString IPAddress;
+
+	UPROPERTY(EditAnywhere, BlueprintReadWrite, Category = Connection)
+	int32 Port = 0;
+
+private:
+	virtual TMap<FString, FPropertyValue> GetCreationParameters() override;
+};
+
+
diff --git a/Engine/Plugins/VirtualProduction/CaptureManager/CaptureManagerApp/Source/CaptureManagerEditor/Private/UI/CaptureSourceObject.cpp b/Engine/Plugins/VirtualProduction/CaptureManager/CaptureManagerApp/Source/CaptureManagerEditor/Private/UI/CaptureSourceObject.cpp
new file mode 100644
--- /dev/null
+++ b/Engine/Plugins/VirtualProduction/CaptureManager/CaptureManagerApp/Source/CaptureManagerEditor/Private/UI/CaptureSourceObject.cpp
@@ -0,0 +1,131 @@
+// Copyright Epic Games, Inc. All Rights Reserved.
+
+#include "CaptureSourceObject.h"
+
+namespace UE::CaptureManager
+{
+
+FCaptureSourceParamValue CreateCaptureSourceParamValue(const FPropertyValue& InPropertyValue)
+{
+	FCaptureSourceParamValue ParamValue;
+
+	if (InPropertyValue.IsType<bool>())
+	{
+		ParamValue.Type = EParamValueType::Bool;
+		ParamValue.BoolValue = InPropertyValue.Get<bool>();
+	}
+	else if (InPropertyValue.IsType<int64>())
+	{
+		ParamValue.Type = EParamValueType::Integer;
+		ParamValue.IntValue = InPropertyValue.Get<int64>();
+	}
+	else if (InPropertyValue.IsType<double>())
+	{
+		ParamValue.Type = EParamValueType::Double;
+		ParamValue.DoubleValue = InPropertyValue.Get<double>();
+	}
+	else if (InPropertyValue.IsType<FString>())
+	{
+		ParamValue.Type = EParamValueType::String;
+		ParamValue.StringValue = InPropertyValue.Get<FString>();
+	}
+	else if (InPropertyValue.IsType<TArray<bool>>())
+	{
+		ParamValue.Type = EParamValueType::ArrayBool;
+		ParamValue.ArrayBoolValue = InPropertyValue.Get<TArray<bool>>();
+	}
+	else if (InPropertyValue.IsType<TArray<int64>>())
+	{
+		ParamValue.Type = EParamValueType::ArrayInteger;
+		ParamValue.ArrayIntValue = InPropertyValue.Get<TArray<int64>>();
+	}
+	else if (InPropertyValue.IsType<TArray<double>>())
+	{
+		ParamValue.Type = EParamValueType::ArrayDouble;
+		ParamValue.ArrayDoubleValue = InPropertyValue.Get<TArray<double>>();
+	}
+	else if (InPropertyValue.IsType<TArray<FString>>())
+	{
+		ParamValue.Type = EParamValueType::ArrayString;
+		ParamValue.ArrayStringValue = InPropertyValue.Get<TArray<FString>>();
+	}
+	else
+	{
+		ParamValue.Type = EParamValueType::Unknown;
+	}
+
+	return ParamValue;
+}
+
+FPropertyValue CreateCaptureSourcePropertyValue(const FCaptureSourceParamValue& InParamValue)
+{
+	FPropertyValue PropertyValue;
+
+	if (InParamValue.Type == EParamValueType::Bool)
+	{
+		PropertyValue.Set<bool>(InParamValue.BoolValue);
+	}
+	else if (InParamValue.Type == EParamValueType::Integer)
+	{
+		PropertyValue.Set<int64>(InParamValue.IntValue);
+	}
+	else if (InParamValue.Type == EParamValueType::Double)
+	{
+		PropertyValue.Set<double>(InParamValue.DoubleValue);
+	}
+	else if (InParamValue.Type == EParamValueType::String)
+	{
+		PropertyValue.Set<FString>(InParamValue.StringValue);
+	}
+	else if (InParamValue.Type == EParamValueType::ArrayBool)
+	{
+		PropertyValue.Set<TArray<bool>>(InParamValue.ArrayBoolValue);
+	}
+	else if (InParamValue.Type == EParamValueType::ArrayInteger)
+	{
+		PropertyValue.Set<TArray<int64>>(InParamValue.ArrayIntValue);
+	}
+	else if (InParamValue.Type == EParamValueType::ArrayDouble)
+	{
+		PropertyValue.Set<TArray<double>>(InParamValue.ArrayDoubleValue);
+	}
+	else if (InParamValue.Type == EParamValueType::ArrayString)
+	{
+		PropertyValue.Set<TArray<FString>>(InParamValue.ArrayStringValue);
+	}
+	else
+	{
+		PropertyValue.Set<FEmptyVariantState>(FEmptyVariantState());
+	}
+
+	return PropertyValue;
+}
+
+FCaptureSourceInfo CreateCaptureSourceInfo(const FString& InFactoryId, const FString& InName, const TMap<FString, FPropertyValue>& InCreationParams)
+{
+	FCaptureSourceInfo CaptureSourceInfo;
+
+	CaptureSourceInfo.FactoryId = InFactoryId;
+	CaptureSourceInfo.Name = InName;
+
+	for (const TPair<FString, FPropertyValue>& CreationParam : InCreationParams)
+	{
+		CaptureSourceInfo.ParamValues.Add(CreationParam.Key, CreateCaptureSourceParamValue(CreationParam.Value));
+	}
+
+	return CaptureSourceInfo;
+}
+
+TMap<FString, FPropertyValue> CreateCaptureSourceInfoParams(const TMap<FString, FCaptureSourceParamValue>& InSourceInfoParams)
+{
+	TMap<FString, FPropertyValue> InfoParams;
+
+	for (const TPair<FString, FCaptureSourceParamValue>& ParamValue : InSourceInfoParams)
+	{
+		InfoParams.Add(ParamValue.Key, CreateCaptureSourcePropertyValue(ParamValue.Value));
+	}
+
+	return InfoParams;
+}
+
+}
\ No newline at end of file
diff --git a/Engine/Plugins/VirtualProduction/CaptureManager/CaptureManagerApp/Source/CaptureManagerEditor/Private/UI/CaptureSourceObject.h b/Engine/Plugins/VirtualProduction/CaptureManager/CaptureManagerApp/Source/CaptureManagerEditor/Private/UI/CaptureSourceObject.h
new file mode 100644
--- /dev/null
+++ b/Engine/Plugins/VirtualProduction/CaptureManager/CaptureManagerApp/Source/CaptureManagerEditor/Private/UI/CaptureSourceObject.h
@@ -0,0 +1,82 @@
+// Copyright Epic Games, Inc. All Rights Reserved.
+
+#pragma once
+
+#include "UObject/ObjectMacros.h"
+
+#include "CaptureSource.h"
+
+#include "CaptureSourceObject.generated.h"
+
+UENUM()
+enum class EParamValueType
+{
+	Unknown,
+	Bool,
+	Integer,
+	Double,
+	String,
+	ArrayBool,
+	ArrayInteger,
+	ArrayDouble,
+	ArrayString
+};
+
+USTRUCT()
+struct FCaptureSourceParamValue
+{
+	GENERATED_BODY()
+
+	UPROPERTY()
+	EParamValueType Type;
+
+	UPROPERTY()
+	bool BoolValue;
+
+	UPROPERTY()
+	int64 IntValue;
+
+	UPROPERTY()
+	double DoubleValue;
+
+	UPROPERTY()
+	FString StringValue;
+
+	UPROPERTY()
+	TArray<bool> ArrayBoolValue;
+
+	UPROPERTY()
+	TArray<int64> ArrayIntValue;
+
+	UPROPERTY()
+	TArray<double> ArrayDoubleValue;
+
+	UPROPERTY()
+	TArray<FString> ArrayStringValue;
+};
+
+USTRUCT()
+struct FCaptureSourceInfo
+{
+	GENERATED_BODY()
+
+	UPROPERTY()
+	FString Name;
+
+	UPROPERTY()
+	FString FactoryId;
+
+	UPROPERTY()
+	TMap<FString, FCaptureSourceParamValue> ParamValues;
+};
+
+
+namespace UE::CaptureManager
+{
+
+FCaptureSourceParamValue CreateCaptureSourceParamValue(const FPropertyValue& InPropertyValue);
+FPropertyValue CreateCaptureSourcePropertyValue(const FCaptureSourceParamValue& InParamValue);
+FCaptureSourceInfo CreateCaptureSourceInfo(const FString& InFactoryId, const FString& InName, const TMap<FString, FPropertyValue>& InCreationParams);
+TMap<FString, FPropertyValue> CreateCaptureSourceInfoParams(const TMap<FString, FCaptureSourceParamValue>& InSourceInfoParams);
+
+}
diff --git a/Engine/Plugins/VirtualProduction/CaptureManager/CaptureManagerApp/Source/CaptureManagerEditor/Private/UI/LiveLinkHubCaptureManagerMode.cpp b/Engine/Plugins/VirtualProduction/CaptureManager/CaptureManagerApp/Source/CaptureManagerEditor/Private/UI/LiveLinkHubCaptureManagerMode.cpp
new file mode 100644
--- /dev/null
+++ b/Engine/Plugins/VirtualProduction/CaptureManager/CaptureManagerApp/Source/CaptureManagerEditor/Private/UI/LiveLinkHubCaptureManagerMode.cpp
@@ -0,0 +1,280 @@
+// Copyright Epic Games, Inc. All Rights Reserved.
+
+#include "LiveLinkHubCaptureManagerMode.h"
+
+
+#include "CaptureCommunicationManagerModule.h"
+#include "CaptureManagerPanelController.h"
+#include "CaptureUtilsModule.h"
+#include "LiveLinkHubDiscoveryModule.h"
+#include "Network/NetworkMisc.h"
+#include "SIngestJobProcessor.h"
+#include "TakesView.h"
+#include "WorkflowOrientedApp/WorkflowCentricApplication.h"
+#include "WorkflowOrientedApp/WorkflowTabFactory.h"
+#include "WorkspaceMenuStructure.h"
+#include "WorkspaceMenuStructureModule.h"
+#include "Widgets/Docking/SDockTab.h"
+
+#define LOCTEXT_NAMESPACE "LiveLinkHubCaptureManagerMode"
+
+namespace UE::CaptureManager
+{
+	const FName TakeBrowserTabId("TakeBrowserTabId");
+	const FName JobsListTabId("JobsListTabId");
+	const FName JobDetailsTabId("JobDetailsTabId");
+	const FName StatusBarId("StatusBarId");
+
+	static TSharedRef<UE::CaptureManager::FDiscoveryRequester> GetDiscoveryRequester()
+	{
+		FLiveLinkHubDiscovery& LiveLinkHubDiscovery =
+			FModuleManager::LoadModuleChecked<FLiveLinkHubDiscovery>("LiveLinkHubDiscovery");
+
+		return LiveLinkHubDiscovery.GetDiscoveryRequster();
+	}
+
+	static TSharedRef<FCommunicationManager> GetCommunicationManager()
+	{
+		FCaptureCommunicationManagerModule& Module =
+			FModuleManager::Get().LoadModuleChecked<FCaptureCommunicationManagerModule>("CaptureCommunicationManager");
+
+		return Module.GetManager();
+	}
+
+	static TSharedRef<FCaptureTimerManager> GetTimerManager()
+	{
+		FCaptureUtilsModule& Module =
+			FModuleManager::Get().LoadModuleChecked<FCaptureUtilsModule>("CaptureUtils");
+
+		return Module.GetTimerManager();
+	}
+}
+
+const FText FLiveLinkHubCaptureManagerMode::TakeBrowserTabName = LOCTEXT("TakeBrowserTabLabel", "Take Browser");
+const FText FLiveLinkHubCaptureManagerMode::JobsListTabName = LOCTEXT("JobsListTabLabel", "Jobs");
+const FText FLiveLinkHubCaptureManagerMode::JobDetailsTabName = LOCTEXT("JobDetailsTabLabel", "Job Details");
+
+struct FTakeBrowserTabSummoner : public FWorkflowTabFactory
+{
+public:
+	FTakeBrowserTabSummoner(TSharedPtr<FAssetEditorToolkit> InHostingApp, TSharedPtr<FCaptureManagerPanelController> InPanelController)
+		: FWorkflowTabFactory(UE::CaptureManager::TakeBrowserTabId, MoveTemp(InHostingApp))
+		, PanelController(MoveTemp(InPanelController))
+	{
+		TabLabel = LOCTEXT("TakeBrowserTabLabel", "Take Browser");
+		TabIcon = FSlateIcon(FAppStyle::GetAppStyleSetName(), "LevelEditor.Tabs.Details");
+
+		bIsSingleton = true;
+	}
+
+	virtual TSharedRef<SWidget> CreateTabBody(const FWorkflowTabSpawnInfo& Info) const override
+	{
+		return PanelController->GetTakesView().ToSharedRef();
+	}
+
+private:
+	TSharedPtr<FCaptureManagerPanelController> PanelController;
+};
+
+struct FJobsListTabSummoner : public FWorkflowTabFactory
+{
+public:
+	FJobsListTabSummoner(TSharedPtr<FAssetEditorToolkit> InHostingApp, TSharedPtr<FCaptureManagerPanelController> InPanelController)
+		: FWorkflowTabFactory(UE::CaptureManager::JobsListTabId, MoveTemp(InHostingApp))
+		, PanelController(MoveTemp(InPanelController))
+	{
+		TabLabel = LOCTEXT("JobsListTabLabel", "Jobs List");
+		TabIcon = FSlateIcon(FAppStyle::GetAppStyleSetName(), "LevelEditor.Tabs.Details");
+
+		bIsSingleton = true;
+	}
+
+	virtual ~FJobsListTabSummoner() 
+	{
+	}
+
+	virtual TSharedRef<SWidget> CreateTabBody(const FWorkflowTabSpawnInfo& Info) const override
+	{
+		return PanelController->GetIngestJobProcessorWidget();
+	}
+
+private:
+	TSharedPtr<FCaptureManagerPanelController> PanelController;
+};
+
+struct FJobsDetailsTabSummoner : public FWorkflowTabFactory
+{
+public:
+	FJobsDetailsTabSummoner(TSharedPtr<FAssetEditorToolkit> InHostingApp, TSharedPtr<FCaptureManagerPanelController> InPanelController)
+		: FWorkflowTabFactory(UE::CaptureManager::JobDetailsTabId, MoveTemp(InHostingApp))
+		, PanelController(MoveTemp(InPanelController))
+	{
+		TabLabel = LOCTEXT("JobDetailsTabLabel", "Job Details");
+		TabIcon = FSlateIcon(FAppStyle::GetAppStyleSetName(), "LevelEditor.Tabs.Details");
+
+		bIsSingleton = true;
+	}
+
+	virtual ~FJobsDetailsTabSummoner() {}
+
+	virtual TSharedRef<SWidget> CreateTabBody(const FWorkflowTabSpawnInfo& Info) const override
+	{
+		return PanelController->GetIngestJobDetailsWidget();
+	}
+
+private:
+	TSharedPtr<FCaptureManagerPanelController> PanelController;
+};
+
+FLiveLinkHubCaptureManagerMode::FLiveLinkHubCaptureManagerMode(TSharedPtr<class FWorkflowCentricApplication> App)
+	: FLiveLinkHubApplicationMode("CaptureManager", LOCTEXT("CaptureManagerModeLabel", "Capture Manager"), App)
+	, TimerManager(UE::CaptureManager::GetTimerManager())
+	, DiscoveryRequester(UE::CaptureManager::GetDiscoveryRequester())
+	, CommunicationManager(UE::CaptureManager::GetCommunicationManager())
+{
+	using namespace UE::CaptureManager;
+
+	TabLayout = FTabManager::NewLayout("LiveLinkCaptureManagerMode_v1.0");
+	
+	TOptional<FString> LocalIpAddressOpt = GetLocalIpAddress();
+
+	check(LocalIpAddressOpt.IsSet());
+
+	if (LocalIpAddressOpt.IsSet())
+	{
+		LocalIpAddress = LocalIpAddressOpt.GetValue();
+	}
+
+	using FConnectionStateDelegate = FCommunicationManager::FConnectionStateChanged::FDelegate;
+	ConnectionChangedHandler =
+		CommunicationManager->AddConnectionStateListener(FConnectionStateDelegate::CreateRaw(this, &FLiveLinkHubCaptureManagerMode::OnConnectionStateChanged));
+
+	ConnectionInvoker = TimerManager->AddTimer(
+		FTimerDelegate::CreateRaw(this, &FLiveLinkHubCaptureManagerMode::OnTick, DiscoveryRequester, CommunicationManager), 
+		ConnectionInvokePeriod, 
+		true);
+
+	TabLayout->AddArea
+		(
+			FTabManager::NewPrimaryArea()
+			->SetOrientation(Orient_Vertical)
+			->Split
+			(
+				FTabManager::NewSplitter()
+				->SetSizeCoefficient(1.f)
+				->SetOrientation(Orient_Horizontal)
+				->Split
+				(
+					FTabManager::NewStack()
+					->SetSizeCoefficient(0.5f)
+					->AddTab(TakeBrowserTabId, ETabState::OpenedTab)
+				)
+				->Split
+				(
+					FTabManager::NewSplitter()
+					->SetSizeCoefficient(0.25f)
+					->SetOrientation(Orient_Vertical)
+					->Split
+					(
+						FTabManager::NewStack()
+						->SetSizeCoefficient(0.5f)
+						->AddTab(JobsListTabId, ETabState::OpenedTab)
+					)
+					->Split
+					(
+						FTabManager::NewStack()
+						->SetSizeCoefficient(0.5f)
+						->AddTab(JobDetailsTabId, ETabState::OpenedTab)
+					)
+				)
+			)
+		);
+	
+	
+	PanelController = FCaptureManagerPanelController::MakeInstance();
+
+	TabFactories.RegisterFactory(MakeShared<FTakeBrowserTabSummoner>(App, PanelController));
+	TabFactories.RegisterFactory(MakeShared<FJobsListTabSummoner>(App, PanelController));
+	TabFactories.RegisterFactory(MakeShared<FJobsDetailsTabSummoner>(App, PanelController));
+}
+
+FLiveLinkHubCaptureManagerMode::~FLiveLinkHubCaptureManagerMode()
+{
+	CommunicationManager->RemoveConnectionStateListener(ConnectionChangedHandler);
+	TimerManager->RemoveTimer(ConnectionInvoker);
+}
+
+TArray<TSharedRef<SWidget>> FLiveLinkHubCaptureManagerMode::GetStatusBarWidgets_Impl()
+{
+
+	return {
+		SNew(SBox)
+		.Padding(0, 4.0, 0, 0)
+		[
+			SNew(STextBlock)
+			.Text(this, &FLiveLinkHubCaptureManagerMode::GetConnectionStateText)
+		]
+	};
+}
+
+
+void FLiveLinkHubCaptureManagerMode::RegisterTabFactories(TSharedPtr<FTabManager> InTabManager)
+{
+	if (TSharedPtr<FWorkflowCentricApplication> App = WeakApp.Pin())
+	{
+		// Pushing the tab factories will register the tab spwaners for the current mode.
+		App->PushTabFactories(TabFactories);
+	}
+
+	FApplicationMode::RegisterTabFactories(InTabManager);
+}
+
+void FLiveLinkHubCaptureManagerMode::OnTick(TSharedRef<UE::CaptureManager::FDiscoveryRequester> InDiscoveryRequester,
+											 TSharedRef<FCommunicationManager> InCommunicationManager)
+{
+	using namespace UE::CaptureManager;
+
+	TArray<FDiscoveredClient> DiscoveredClients = InDiscoveryRequester->GetDiscoveredClients();
+
+	if (InCommunicationManager->IsConnected() || DiscoveredClients.IsEmpty())
+	{
+		return;
+	}
+
+	FDiscoveredClient* LocalActiveClient = DiscoveredClients.FindByPredicate(
+		[this](const FDiscoveredClient& InClient)
+		{
+			return InClient.GetIPAddress() == LocalIpAddress && InClient.IsActive();
+		}
+	);
+
+	if (LocalActiveClient)
+	{
+		InCommunicationManager->Connect(LocalActiveClient->GetMessageAddress(), LocalActiveClient->GetIPAddress(), LocalActiveClient->GetExportPort());
+	}
+}
+
+void FLiveLinkHubCaptureManagerMode::OnConnectionStateChanged(FMessageAddress InClientAddress, bool bInIsConnected)
+{
+	if (bInIsConnected)
+	{
+		TimerManager->RemoveTimer(ConnectionInvoker);
+	}
+	else
+	{
+		ConnectionInvoker = TimerManager->AddTimer(
+			FTimerDelegate::CreateRaw(this, &FLiveLinkHubCaptureManagerMode::OnTick, DiscoveryRequester, CommunicationManager),
+			ConnectionInvokePeriod, 
+			true);
+	}
+
+	bIsConnected = bInIsConnected;
+}
+
+
+FText FLiveLinkHubCaptureManagerMode::GetConnectionStateText() const
+{
+	return bIsConnected ? LOCTEXT("SStatusBarView_Connected", "Connected") : LOCTEXT("SStatusBarView_Disconnected", "Disconnected");
+}
+
+#undef LOCTEXT_NAMESPACE
\ No newline at end of file
diff --git a/Engine/Plugins/VirtualProduction/CaptureManager/CaptureManagerApp/Source/CaptureManagerEditor/Private/UI/LiveLinkHubCaptureManagerMode.h b/Engine/Plugins/VirtualProduction/CaptureManager/CaptureManagerApp/Source/CaptureManagerEditor/Private/UI/LiveLinkHubCaptureManagerMode.h
new file mode 100644
--- /dev/null
+++ b/Engine/Plugins/VirtualProduction/CaptureManager/CaptureManagerApp/Source/CaptureManagerEditor/Private/UI/LiveLinkHubCaptureManagerMode.h
@@ -0,0 +1,79 @@
+// Copyright Epic Games, Inc. All Rights Reserved.
+
+#pragma once
+
+#include "LiveLinkHubApplicationMode.h"
+
+#include "Framework/Docking/TabManager.h"
+#include "Templates/SharedPointer.h"
+
+#include "Async/CaptureTimerManager.h"
+
+#include "IMessageContext.h"
+#include "WorkflowOrientedApp/WorkflowTabManager.h"
+
+namespace UE::CaptureManager
+{
+	class FDiscoveryRequester;
+}
+
+class FCommunicationManager;
+
+class SCaptureManagerMainTabView;
+class SDockTab;
+
+class FLiveLinkHubCaptureManagerMode : public FLiveLinkHubApplicationMode
+{
+public:
+	FLiveLinkHubCaptureManagerMode(TSharedPtr<class FWorkflowCentricApplication> App);
+	virtual ~FLiveLinkHubCaptureManagerMode();
+
+	//~ Begin FLiveLinkHubApplicationMode interface
+	virtual FSlateIcon GetModeIcon() const override
+	{
+		return FSlateIcon(TEXT("MetaHumanIdentityStyle"), TEXT("ClassIcon.FootageCaptureData"), TEXT("ClassIcon.FootageCaptureData"));
+	}
+
+	/** Get widgets that should appear in the status bar. */
+	virtual TArray<TSharedRef<SWidget>> GetStatusBarWidgets_Impl() override;
+
+	//~ End FLiveLinkHubApplicationMode interface
+
+	//~ Begin FApplicationMode interface
+	virtual void RegisterTabFactories(TSharedPtr<FTabManager> InTabManager) override;
+	//~ End FApplicationMode interace
+
+private:
+
+	static constexpr float ConnectionInvokePeriod = 3.0f;
+
+	void OnConnectionStateChanged(FMessageAddress InClientAddress, bool bIsConnected);
+	void OnTick(TSharedRef<UE::CaptureManager::FDiscoveryRequester> InDiscoveryRequester, 
+				TSharedRef<FCommunicationManager> InCommunicationManager);
+
+	
+	FText GetConnectionStateText() const;
+
+private:
+	TSharedRef<UE::CaptureManager::FCaptureTimerManager> TimerManager;
+
+	FString LocalIpAddress;
+	TSharedRef<UE::CaptureManager::FDiscoveryRequester> DiscoveryRequester;
+
+	UE::CaptureManager::FCaptureTimerManager::FTimerHandle ConnectionInvoker;
+	TSharedRef<FCommunicationManager> CommunicationManager;
+	FDelegateHandle ConnectionChangedHandler;
+
+	TSharedPtr<class FCaptureManagerPanelController> PanelController;
+
+	static const FText CaptureSourcesTabName;
+	static const FText TakeBrowserTabName;
+	static const FText JobsListTabName;
+	static const FText JobDetailsTabName;
+
+	/** Set of spawnable tabs in the mode */
+	FWorkflowAllowedTabSet TabFactories;
+	
+	std::atomic<bool> bIsConnected = false;
+};
+
diff --git a/Engine/Plugins/VirtualProduction/CaptureManager/CaptureManagerApp/Source/CaptureManagerEditor/Private/UI/TakesView.cpp b/Engine/Plugins/VirtualProduction/CaptureManager/CaptureManagerApp/Source/CaptureManagerEditor/Private/UI/TakesView.cpp
new file mode 100644
--- /dev/null
+++ b/Engine/Plugins/VirtualProduction/CaptureManager/CaptureManagerApp/Source/CaptureManagerEditor/Private/UI/TakesView.cpp
@@ -0,0 +1,681 @@
+// Copyright Epic Games, Inc. All Rights Reserved.
+
+#include "TakesView.h"
+
+#include "CaptureManagerPanelController.h"
+#include "SIngestJobProcessor.h"
+#include "ImageUtils.h"
+#include "CaptureSourceVirtualAsset.h"
+#include "TakeThumbnail.h"
+#include "ContentBrowserModule.h"
+#include "Async/Monitor.h"
+#include "Misc/FileHelper.h"
+
+#include "Widgets/Input/SButton.h"
+#include "Widgets/Input/SComboBox.h"
+
+#include "Ingest/LiveLinkDeviceCapability_Ingest.h"
+#include "Ingest/IngestCapability_Events.h"
+#include "Ingest/IngestCapability_TakeInformation.h"
+
+#define LOCTEXT_NAMESPACE "CaptureManagerPanelViews"
+
+DEFINE_LOG_CATEGORY_STATIC(LogTakesView, Log, All);
+
+static const FText JobPresetText = LOCTEXT("JobPresetText", "Job Preset");
+
+static const FString OuterPackageName = TEXT("CaptureDevices");
+
+STakesView::STakesView() = default;
+
+void STakesView::Construct(
+	const FArguments& InArgs, 
+	TSharedRef<UE::CaptureManager::FIngestJobSettingsManager> InIngestJobsSettingsManager,
+	TUniquePtr<FGetCurrentSelectionDelegate> InGetCurrentSelectionDelegate
+)
+{
+	GetCurrentSelectionDelegate = MoveTemp(InGetCurrentSelectionDelegate);
+	check(GetCurrentSelectionDelegate);
+
+	IngestJobSettingsManager = InIngestJobsSettingsManager.ToSharedPtr();
+	OnAddTakesToIngestQueue = InArgs._OnAddTakesToIngestQueue;
+	CurrentPresetName = IngestJobSettingsManager->GetSelectedPreset().DisplayName;
+
+	TArray<FText> PresetDisplayNames = IngestJobSettingsManager->GetPresetDisplayNames();
+	PresetDisplayNames.Sort(FText::FSortPredicate());
+
+	for (const FText& PresetDisplayName : PresetDisplayNames)
+	{
+		IngestJobPresetNames.Emplace(MakeShared<FText>(PresetDisplayName));
+	}
+
+	AddToQueueButton = SNew(SButton)
+		.HAlign(EHorizontalAlignment::HAlign_Center)
+		.Text(LOCTEXT("AddToQueueText", "Add to Queue"))
+		.OnClicked_Raw(this, &STakesView::OnAddToQueueButtonClicked);
+
+	IContentBrowserSingleton& ContentBrowser = FModuleManager::LoadModuleChecked<FContentBrowserModule>("ContentBrowser").Get();
+
+	FAssetPickerConfig DevicesPickerConfig = InArgs._DevicesPickerConfig;
+	FAssetPickerConfig TakesPickerConfig = InArgs._TakesPickerConfig;
+
+	ChildSlot
+		[
+			SNew(SVerticalBox)
+				+ SVerticalBox::Slot()
+				[
+					SNew(SSplitter)
+						.Orientation(Orient_Horizontal)
+						+ SSplitter::Slot()
+						.Value(0.25f)
+						[
+							ContentBrowser.CreateAssetPicker(DevicesPickerConfig)
+						]
+						+ SSplitter::Slot()
+						.Value(0.75f)
+						[
+							ContentBrowser.CreateAssetPicker(TakesPickerConfig)
+						]
+				]
+				+ SVerticalBox::Slot()
+				.VAlign(EVerticalAlignment::VAlign_Bottom)
+				.AutoHeight()
+				.Padding(2.0f)
+				[
+					SNew(SHorizontalBox)
+						+ SHorizontalBox::Slot()
+						.Padding(2.0f)
+						.AutoWidth()
+						.VAlign(EVerticalAlignment::VAlign_Center)
+						[
+							SNew(STextBlock)
+								.Text(JobPresetText)
+						]
+						+ SHorizontalBox::Slot()
+						.Padding(2.0f)
+						.HAlign(EHorizontalAlignment::HAlign_Fill)
+						.VAlign(EVerticalAlignment::VAlign_Center)
+						[
+							SNew(SComboBox<TSharedRef<FText>>)
+								.OptionsSource(&IngestJobPresetNames)
+								.OnGenerateWidget_Raw(this, &STakesView::OnGeneratePresetNameWidget)
+								.OnSelectionChanged_Raw(this, &STakesView::OnPresetSelectionChanged)
+								.Content()
+								[
+									SNew(STextBlock)
+										.MinDesiredWidth(200)
+										.Text_Lambda(
+											[this]()
+											{
+												return CurrentPresetName;
+											}
+										)
+								]
+						]
+						+ SHorizontalBox::Slot()
+						.Padding(2.0f)
+						.AutoWidth()
+						.VAlign(EVerticalAlignment::VAlign_Center)
+						[
+							AddToQueueButton.ToSharedRef()
+						]
+				]
+		];
+}
+
+void STakesView::SetAddToQueueButtonEnabled(const bool bIsEnabled)
+{
+	check(AddToQueueButton);
+
+	if (AddToQueueButton)
+	{
+		AddToQueueButton->SetEnabled(bIsEnabled);
+	}
+}
+
+TArray<TStrongObjectPtr<UTakeVirtualAsset>> STakesView::GetSelectedTakeAssets() const
+{
+	check(GetCurrentSelectionDelegate);
+
+	if (!GetCurrentSelectionDelegate)
+	{
+		return {};
+	}
+
+	const TArray<FAssetData> SelectedAssetDatas = GetCurrentSelectionDelegate->Execute();
+
+	TArray<TStrongObjectPtr<UTakeVirtualAsset>> TakeAssets;
+	TakeAssets.Reserve(SelectedAssetDatas.Num());
+
+	for (const FAssetData& AssetData : SelectedAssetDatas)
+	{
+		TStrongObjectPtr<UTakeVirtualAsset> TakeAsset(Cast<UTakeVirtualAsset>(AssetData.GetAsset()));
+
+		if (TakeAsset)
+		{
+			TakeAssets.Emplace(MoveTemp(TakeAsset));
+		}
+	}
+
+	return TakeAssets;
+}
+
+FReply STakesView::OnAddToQueueButtonClicked()
+{
+	check(IngestJobSettingsManager);
+
+	if (IngestJobSettingsManager)
+	{
+		TArray<TStrongObjectPtr<UTakeVirtualAsset>> SelectedTakeAssets = GetSelectedTakeAssets();
+
+		// If the user has selected a lot of takes, rather than add them in a random order, we add them in the order of
+		// acquisition, just so that it is deterministic
+		SelectedTakeAssets.Sort(
+			[](const TStrongObjectPtr<const UTakeVirtualAsset>& LeftItem, const TStrongObjectPtr<const UTakeVirtualAsset>& RightItem)
+			{
+				return LeftItem->Metadata.DateTime.Get(FDateTime()) < RightItem->Metadata.DateTime.Get(FDateTime());
+			}
+		);
+
+		const UIngestJobSettingsBase& SelectedPreset = IngestJobSettingsManager->GetSelectedPreset();
+
+		OnAddTakesToIngestQueue.ExecuteIfBound(SelectedTakeAssets, SelectedPreset);
+	}
+
+	return FReply::Handled();
+}
+
+TSharedRef<SWidget> STakesView::OnGeneratePresetNameWidget(TSharedRef<FText> InText)
+{
+	return SNew(STextBlock).Text(*InText);
+}
+
+void STakesView::OnPresetSelectionChanged(TSharedPtr<FText> InText, ESelectInfo::Type InSelectType)
+{
+	check(IngestJobSettingsManager);
+
+	if (InText)
+	{
+		if (IngestJobSettingsManager->SelectPresetByDisplayName(*InText))
+		{
+			CurrentPresetName = *InText;
+		}
+	}
+}
+
+TArray<FAssetData> FTakesView::CreateTakeAssetsDataForCaptureDevice(TOptional<FString> InCaptureDeviceName)
+{
+	using namespace UE::CaptureManager;
+
+	TMap<FString, TArray<TStrongObjectPtr<UTakeVirtualAsset>>> SelectedDevicesTakesMap;
+
+	TSharedPtr<FCaptureManagerPanelController> Controller = PanelController.Pin();
+	if (Controller)
+	{
+		for (const TPair<FGuid, TArray<TStrongObjectPtr<UTakeVirtualAsset>>>& CSIdTakeTuple : CaptureDeviceTakesMap)
+		{
+			FGuid CaptureDeviceId = CSIdTakeTuple.Key;
+			TObjectPtr<ULiveLinkDevice> Device = Controller->GetCaptureDevice(CaptureDeviceId);
+
+			if (!InCaptureDeviceName.IsSet() || Device.GetName() == InCaptureDeviceName.GetValue())
+			{
+				TArray<TStrongObjectPtr<UTakeVirtualAsset>> Takes;
+
+				for (const TStrongObjectPtr<UTakeVirtualAsset>& Take : CSIdTakeTuple.Value)
+				{
+					Takes.Add(Take);
+				}
+				SelectedDevicesTakesMap.Add(Device.GetName(), MoveTemp(Takes));
+			}
+		}
+	}
+
+	return CreateTakeAssetsData(SelectedDevicesTakesMap);
+}
+
+TArray<FAssetData> FTakesView::CreateTakeAssetsData(const TMap<FString, TArray<TStrongObjectPtr<UTakeVirtualAsset>>>& InSelectedCaptureDevicesTakesMap)
+{
+	TArray<FAssetData> TakeAssetsData;
+	for (const TTuple<FString, TArray<TStrongObjectPtr<UTakeVirtualAsset>>>& CaptureDeviceTakesTuple : InSelectedCaptureDevicesTakesMap)
+	{
+		for (TStrongObjectPtr<UTakeVirtualAsset> TakeItem : CaptureDeviceTakesTuple.Value)
+		{
+			FAssetData TakeAssetData(TakeItem.Get());
+			TakeAssetData.AssetClassPath = FTopLevelAssetPath(*(TakeItem->GetPackage()->GetName()), *TakeItem->GetName());
+
+			TakeAssetsData.Add(TakeAssetData);
+		}
+	}
+	return TakeAssetsData;
+}
+
+FTakesView::FTakesView(TWeakPtr<FCaptureManagerPanelController> InController) : 
+	PanelController(InController)
+{
+	CreateTakesView();
+}
+
+void FTakesView::AddTakesToIngestQueue(const TArray<TStrongObjectPtr<UTakeVirtualAsset>>& InTakeAssets, const UIngestJobSettingsBase& InJobPreset)
+{
+	using namespace UE::CaptureManager;
+
+	TSharedPtr<FCaptureManagerPanelController> Controller = PanelController.Pin();
+
+	if (!Controller)
+	{
+		UE_LOG(LogTakesView, Error, TEXT("Failed to add takes to the ingest queue, controller is not available"));
+		return;
+	}
+
+	TArray<TSharedRef<FIngestJob>> IngestJobs;
+	IngestJobs.Reserve(InTakeAssets.Num());
+
+	for (const TStrongObjectPtr<UTakeVirtualAsset>& TakeAsset : InTakeAssets)
+	{
+		// We convert the UObject based settings into something that the ingest job can truly own
+		FIngestJob::FSettings IngestJobSettings =
+		{
+			.LocalWorkFolder = InJobPreset.LocalWorkFolder.Path,
+			.VideoSettings =
+			{
+				.Format = StaticEnum<EOutputImageFormat>()->GetDisplayNameTextByValue(static_cast<std::underlying_type_t<EOutputImageFormat>>(InJobPreset.ImageFormat)).ToString(),
+				.FileNamePrefix = InJobPreset.ImageFileNamePrefix,
+				.ImagePixelFormat = InJobPreset.ImagePixelFormat
+			},
+			.AudioSettings =
+			{
+				.Format = StaticEnum<EAudioFormat>()->GetDisplayNameTextByValue(static_cast<std::underlying_type_t<EAudioFormat>>(InJobPreset.AudioFormat)).ToString(),
+				.FileNamePrefix = InJobPreset.AudioFileNamePrefix
+			}
+		};
+
+		IngestJobs.Emplace(MakeShared<FIngestJob>(TakeAsset->CaptureDeviceId, TakeAsset->TakeId, TakeAsset->Metadata, MoveTemp(IngestJobSettings)));
+	}
+
+	const int32 ExpectedNumAdded = IngestJobs.Num();
+
+	TSharedRef<SIngestJobProcessor> IngestJobProcessor = Controller->GetIngestJobProcessorWidget();
+	const int32 NumAdded = IngestJobProcessor->AddJobs(MoveTemp(IngestJobs));
+
+	if (NumAdded != ExpectedNumAdded)
+	{
+		UE_LOG(LogTakesView, Error, TEXT("Some ingest jobs were not added to the queue (%d out of %d added)"), NumAdded, ExpectedNumAdded);
+	}
+}
+
+void FTakesView::CreateTakesView()
+{
+	TSharedPtr<FCaptureManagerPanelController> Controller = PanelController.Pin();
+
+	if (!Controller)
+	{
+		return;
+	}
+
+	TSharedRef<UE::CaptureManager::FIngestJobSettingsManager> IngestJobSettingsManager = Controller->GetIngestJobSettingsManager();
+
+	FAssetPickerConfig DevicesPickerConfig;
+	{
+		DevicesPickerConfig.SelectionMode = ESelectionMode::Single;
+		DevicesPickerConfig.InitialAssetViewType = EAssetViewType::Column;
+		DevicesPickerConfig.bFocusSearchBoxWhenOpened = true;
+		DevicesPickerConfig.bAllowNullSelection = false;
+		DevicesPickerConfig.bShowBottomToolbar = true;
+		DevicesPickerConfig.bAutohideSearchBar = false;
+		DevicesPickerConfig.bAllowDragging = false;
+		DevicesPickerConfig.bCanShowClasses = false;
+		DevicesPickerConfig.bShowPathInColumnView = false;
+		DevicesPickerConfig.bSortByPathInColumnView = false;
+		DevicesPickerConfig.bShowTypeInColumnView = false;
+		DevicesPickerConfig.bForceShowEngineContent = true;
+		DevicesPickerConfig.bForceShowPluginContent = true;
+		DevicesPickerConfig.SyncToAssetsDelegates.Add(&SyncToDevicesDelegate);
+		DevicesPickerConfig.RefreshAssetViewDelegates.Add(&RefreshDevicesViewDelegate);
+		DevicesPickerConfig.Filter.ClassPaths.Add(UCaptureSourceVirtualAsset::StaticClass()->GetClassPathName());
+		DevicesPickerConfig.Filter.bRecursiveClasses = true;
+		DevicesPickerConfig.Filter.bRecursivePaths = true;
+
+		DevicesPickerConfig.OnGetCustomSourceAssets =
+			FOnGetCustomSourceAssets::CreateLambda([this](const FARFilter& SourceFilter, TArray<FAssetData>& OutAssets)
+				{
+					for (const TStrongObjectPtr<UCaptureSourceVirtualAsset>& CaptureDevice : CaptureDevices)
+					{
+						FAssetData AssetData(CaptureDevice.Get());
+						OutAssets.Add(MoveTemp(AssetData));
+					}
+				});
+
+		DevicesPickerConfig.OnAssetSelected = FOnAssetSelected::CreateLambda([this](const FAssetData& InSelectedItemData)
+			{
+				if (!InSelectedItemData.GetAsset())
+				{
+					SelectedItem = nullptr;
+					RefreshTakesViewDelegate.ExecuteIfBound(true);
+					return;
+				}
+
+				TStrongObjectPtr<UCaptureSourceVirtualAsset>* CaptureDevice = 
+					CaptureDevices.FindByPredicate([InSelectedItemData](const TStrongObjectPtr<UCaptureSourceVirtualAsset>& InCaptureDevice)
+				{
+					return InCaptureDevice->Id == StaticCast<UCaptureSourceVirtualAsset*>(InSelectedItemData.GetAsset())->Id;
+				});
+
+				if (CaptureDevice)
+				{
+					SelectedItem = CaptureDevice->Get();
+					RefreshTakesViewDelegate.ExecuteIfBound(true);
+				}
+			});
+	}
+
+	TUniquePtr<FGetCurrentSelectionDelegate> GetCurrentSelectionDelegate = MakeUnique<FGetCurrentSelectionDelegate>();
+
+	FAssetPickerConfig TakesPickerConfig;
+	{
+		TakesPickerConfig.SelectionMode = ESelectionMode::Multi;
+		TakesPickerConfig.InitialAssetViewType = EAssetViewType::Tile;
+		TakesPickerConfig.bFocusSearchBoxWhenOpened = true;
+		TakesPickerConfig.bAllowNullSelection = false;
+		TakesPickerConfig.bShowBottomToolbar = true;
+		TakesPickerConfig.bAutohideSearchBar = false;
+		TakesPickerConfig.bAllowDragging = false;
+		TakesPickerConfig.bCanShowClasses = false;
+		TakesPickerConfig.bShowPathInColumnView = false;
+		TakesPickerConfig.bSortByPathInColumnView = false;
+		TakesPickerConfig.bShowTypeInColumnView = false;
+		TakesPickerConfig.bForceShowEngineContent = true;
+		TakesPickerConfig.bForceShowPluginContent = true;
+		TakesPickerConfig.RefreshAssetViewDelegates.Add(&RefreshTakesViewDelegate);
+		TakesPickerConfig.Filter.ClassPaths.Add(UTakeVirtualAsset::StaticClass()->GetClassPathName());
+		TakesPickerConfig.Filter.bRecursiveClasses = true;
+		TakesPickerConfig.Filter.bRecursivePaths = true;
+		TakesPickerConfig.GetCurrentSelectionDelegates.Add(GetCurrentSelectionDelegate.Get());
+
+		TakesPickerConfig.OnGetCustomSourceAssets =
+			FOnGetCustomSourceAssets::CreateLambda([this](const FARFilter& SourceFilter, TArray<FAssetData>& OutAssets)
+				{
+					TOptional<FString> CaptureDeviceName;
+					if (SelectedItem.IsValid())
+					{
+						CaptureDeviceName = SelectedItem->GetName();
+					}
+
+					OutAssets = CreateTakeAssetsDataForCaptureDevice(CaptureDeviceName);
+				});
+	}
+
+	SAssignNew(TakesTileView, STakesView, MoveTemp(IngestJobSettingsManager), MoveTemp(GetCurrentSelectionDelegate))
+		.OnAddTakesToIngestQueue_Raw(this, &FTakesView::AddTakesToIngestQueue)
+		.DevicesPickerConfig(DevicesPickerConfig)
+		.TakesPickerConfig(TakesPickerConfig);
+}
+
+void FTakesView::CaptureDeviceStarted(FGuid InCaptureDeviceId)
+{
+	AsyncTask(ENamedThreads::GameThread, [this, Controller = PanelController.Pin(), InCaptureDeviceId]()
+	{
+		if (!Controller)
+		{
+			return;
+		}
+
+		using namespace UE::CaptureManager;
+
+		TObjectPtr<ULiveLinkDevice> Device = Controller->GetCaptureDevice(InCaptureDeviceId);
+
+		if (!Device)
+		{
+			return;
+		}
+
+		if (!Device->Implements<ULiveLinkDeviceCapability_Ingest>())
+		{
+			return;
+		}
+
+		FString PackageName = GetTransientPackage()->GetPathName() / OuterPackageName;
+		UPackage* Package = CreatePackage(*PackageName);
+
+		FString CaptureDeviceName = Device->GetName();
+		FName Name(CaptureDeviceName);
+
+		if (UCaptureSourceVirtualAsset* FoundObject = FindObject<UCaptureSourceVirtualAsset>(Package, *CaptureDeviceName, true))
+		{
+			Name = MakeUniqueObjectName(Package, UCaptureSourceVirtualAsset::StaticClass(), FName(CaptureDeviceName));
+		}
+		
+		TStrongObjectPtr<UCaptureSourceVirtualAsset> CaptureDeviceObjectPtr(NewObject<UCaptureSourceVirtualAsset>(Package, Name));
+		CaptureDeviceObjectPtr->Id = InCaptureDeviceId;
+
+		CaptureDevices.AddUnique(MoveTemp(CaptureDeviceObjectPtr));
+
+		FIngestUpdateTakeListCallback Callback = FIngestUpdateTakeListCallback::Type::CreateLambda(
+			[this, InCaptureDeviceId](const TArray<FTakeId>& InTakeIds)
+			{
+				CaptureDeviceTakesMap.Remove(InCaptureDeviceId);
+
+				TArray<TStrongObjectPtr<UTakeVirtualAsset>> TakeObjects = MakeTakeObjects(InCaptureDeviceId, InTakeIds);
+
+				CaptureDeviceTakesMap.Add(InCaptureDeviceId, MoveTemp(TakeObjects));
+
+				RefreshDevicesViewDelegate.ExecuteIfBound(true);
+				RefreshTakesViewDelegate.ExecuteIfBound(true);
+			}
+		);
+
+		TScriptInterface<ILiveLinkDeviceCapability_Ingest> CapabilityInterface(Device);
+
+		CapabilityInterface->SubscribeToEvent(FString(FIngestCapability_TakeAddedEvent::Name),
+											  FCaptureEventHandler([this, InCaptureDeviceId](TSharedPtr<const FCaptureEvent> InTakeAddedEvent)
+		{
+			check(IsInGameThread());
+
+			TSharedPtr<const FIngestCapability_TakeAddedEvent> Event = 
+				StaticCastSharedPtr<const FIngestCapability_TakeAddedEvent>(InTakeAddedEvent);
+
+			TArray<FTakeId> TakeIds = { Event->TakeId };
+			TArray<TStrongObjectPtr<UTakeVirtualAsset>> NewTakeObject = MakeTakeObjects(InCaptureDeviceId, { TakeIds });
+
+			if (TArray<TStrongObjectPtr<UTakeVirtualAsset>>* TakeObjects = CaptureDeviceTakesMap.Find(InCaptureDeviceId))
+			{
+				TakeObjects->Append(NewTakeObject.GetData(), NewTakeObject.Num());
+
+				RefreshTakesViewDelegate.ExecuteIfBound(true);
+			}
+		}));
+
+		CapabilityInterface->SubscribeToEvent(FString(FIngestCapability_TakeUpdatedEvent::Name),
+											  FCaptureEventHandler([this, InCaptureDeviceId](TSharedPtr<const FCaptureEvent> InTakeUpdatedEvent)
+		{
+			check(IsInGameThread());
+
+			TSharedPtr<const FIngestCapability_TakeUpdatedEvent> Event = 
+				StaticCastSharedPtr<const FIngestCapability_TakeUpdatedEvent>(InTakeUpdatedEvent);
+
+			TArray<FTakeId> TakeIds = { Event->TakeId };
+			TArray<TStrongObjectPtr<UTakeVirtualAsset>> NewTakeObject = MakeTakeObjects(InCaptureDeviceId, { TakeIds });
+			check(NewTakeObject.Num() == 1);
+
+			if (TArray<TStrongObjectPtr<UTakeVirtualAsset>>* TakeObjects = CaptureDeviceTakesMap.Find(InCaptureDeviceId))
+			{
+				for (TStrongObjectPtr<UTakeVirtualAsset>& TakeVirtualAsset : *TakeObjects)
+				{
+					if (TakeVirtualAsset->TakeId == Event->TakeId)
+					{
+						TakeVirtualAsset = NewTakeObject[0];
+						break;
+					}
+				}
+
+				RefreshTakesViewDelegate.ExecuteIfBound(true);
+			}
+		}));
+
+		CapabilityInterface->SubscribeToEvent(FString(FIngestCapability_TakeRemovedEvent::Name),
+											  FCaptureEventHandler([this, InCaptureDeviceId](TSharedPtr<const FCaptureEvent> InTakeRemovedEvent)
+		{
+			check(IsInGameThread());
+
+			TSharedPtr<const FIngestCapability_TakeRemovedEvent> Event = 
+				StaticCastSharedPtr<const FIngestCapability_TakeRemovedEvent>(InTakeRemovedEvent);
+
+			if (TArray<TStrongObjectPtr<UTakeVirtualAsset>>* TakeObjects = CaptureDeviceTakesMap.Find(InCaptureDeviceId))
+			{
+				int32 NumRemoved = TakeObjects->RemoveAll([this, TakeId = Event->TakeId](const TStrongObjectPtr<UTakeVirtualAsset>& InObject)
+				{
+					return InObject->TakeId == TakeId;
+				});
+
+				if (NumRemoved > 0)
+				{
+					RefreshTakesViewDelegate.ExecuteIfBound(true);
+				}
+			}
+		}));
+		
+		ILiveLinkDeviceCapability_Ingest::Execute_CreateUpdateTakeListCallback(Device)->Callback = MoveTemp(Callback);
+		ILiveLinkDeviceCapability_Ingest::Execute_UpdateTakeListAsync(Device);
+	});
+}
+
+void FTakesView::CaptureDeviceStopped(FGuid InCaptureDeviceId)
+{
+	AsyncTask(ENamedThreads::GameThread, [this, Controller = PanelController.Pin(), InCaptureDeviceId]()
+	{
+		if (!Controller)
+		{
+			return;
+		}
+
+		using namespace UE::CaptureManager;
+
+		TObjectPtr<ULiveLinkDevice> Device = Controller->GetCaptureDevice(InCaptureDeviceId);
+
+		if (!Device)
+		{
+			return;
+		}
+
+		CaptureDeviceTakesMap.Remove(InCaptureDeviceId);
+		CaptureDevices.RemoveAll([InCaptureDeviceId](const TStrongObjectPtr<UCaptureSourceVirtualAsset>& InCaptureDevice)
+		{
+			return InCaptureDevice->Id == InCaptureDeviceId;
+		});
+
+		RefreshDevicesViewDelegate.ExecuteIfBound(true);
+		RefreshTakesViewDelegate.ExecuteIfBound(true);
+	});
+}
+
+void FTakesView::CaptureDeviceToBeRemoved(FGuid InCaptureDeviceId)
+{
+	AsyncTask(ENamedThreads::GameThread, [this, InCaptureDeviceId]()
+	{
+		CaptureDeviceTakesMap.Remove(InCaptureDeviceId);
+		CaptureDevices.RemoveAll([InCaptureDeviceId](const TStrongObjectPtr<UCaptureSourceVirtualAsset>& InCaptureDevice)
+		{
+			return InCaptureDevice->Id == InCaptureDeviceId;
+		});
+
+		RefreshDevicesViewDelegate.ExecuteIfBound(true);
+		RefreshTakesViewDelegate.ExecuteIfBound(true);
+	});
+}
+
+TArray<TStrongObjectPtr<UTakeVirtualAsset>> FTakesView::MakeTakeObjects(FGuid InCaptureDeviceId, TOptional<TArray<UE::CaptureManager::FTakeId>> InTakeIds) {
+	TArray<TStrongObjectPtr<UTakeVirtualAsset>> TakeObjects;
+
+	TSharedPtr<FCaptureManagerPanelController> Controller = PanelController.Pin();
+	if (!Controller)
+	{
+		return TakeObjects;
+	}
+
+	using namespace UE::CaptureManager;
+
+	TObjectPtr<ULiveLinkDevice> Device = Controller->GetCaptureDevice(InCaptureDeviceId);
+
+	if (!Device)
+	{
+		return TakeObjects;
+	}
+
+	if (!Device->Implements<ULiveLinkDeviceCapability_Ingest>())
+	{
+		return TakeObjects;
+	}
+
+	const FString& CaptureDeviceName = Device.GetName();
+
+	TScriptInterface<ILiveLinkDeviceCapability_Ingest> IngestInterface(Device);
+
+	TArray<FTakeId> TakeIds;
+	if (InTakeIds.IsSet())
+	{
+		TakeIds = InTakeIds.GetValue();
+	}
+	else
+	{
+		TakeIds = IngestInterface->GetTakeIdentifiers();
+	}
+
+	FString PackageName = GetTransientPackage()->GetPathName() / OuterPackageName / CaptureDeviceName;
+	UPackage* Package = CreatePackage(*PackageName);
+
+	for (const FTakeId& TakeId : TakeIds)
+	{
+		TOptional<FTakeMetadata> TakeInfoOpt = IngestInterface->GetTakeMetadata(TakeId);
+		if (!TakeInfoOpt.IsSet())
+		{
+			continue;
+		}
+
+		FTakeMetadata TakeInfo = TakeInfoOpt.GetValue();
+
+		FName DesiredName(TakeInfo.Slate, TakeInfo.TakeNumber);
+		FName Name = MakeUniqueObjectName(Package, UTakeVirtualAsset::StaticClass(), FName(DesiredName));
+
+		TStrongObjectPtr<UTakeVirtualAsset> TakeObject(NewObject<UTakeVirtualAsset>(Package, Name));
+		TakeObject->CaptureDeviceId = InCaptureDeviceId;
+		TakeObject->TakeId = TakeId;
+		TakeObject->Metadata = MoveTemp(TakeInfo);
+
+		TOptional<FTakeThumbnail> TakeThumbnail = CreateThumbnail(TakeObject->Metadata.Thumbnail);
+
+		if (TakeThumbnail.IsSet())
+		{
+			TakeObject->Thumbnail = MoveTemp(TakeThumbnail.GetValue());
+		}
+
+		TakeObjects.Add(MoveTemp(TakeObject));
+	}
+
+	return TakeObjects;
+}
+
+TOptional<FTakeThumbnail> FTakesView::CreateThumbnail(const FTakeThumbnailData& InThumbnailData)
+{
+	TOptional<TArray<uint8>> ThumbnailRawDataOpt = InThumbnailData.GetThumbnailData();
+
+	if (!ThumbnailRawDataOpt.IsSet())
+	{
+		return {};
+	}
+
+	TArray<uint8> ThumbnailRawData = MoveTemp(ThumbnailRawDataOpt.GetValue());
+
+	FImage Image;
+	FImageUtils::DecompressImage(ThumbnailRawData.GetData(), ThumbnailRawData.Num(), Image);
+	Image.ChangeFormat(ERawImageFormat::BGRA8, EGammaSpace::sRGB);
+
+	UTexture2D* PreviewImageTexture = FImageUtils::CreateTexture2DFromImage(Image);
+	if (!PreviewImageTexture)
+	{
+		return {};
+	}
+
+	return FTakeThumbnail(PreviewImageTexture);
+}
+
+
+#undef LOCTEXT_NAMESPACE
diff --git a/Engine/Plugins/VirtualProduction/CaptureManager/CaptureManagerApp/Source/CaptureManagerEditor/Private/UI/TakesView.h b/Engine/Plugins/VirtualProduction/CaptureManager/CaptureManagerApp/Source/CaptureManagerEditor/Private/UI/TakesView.h
new file mode 100644
--- /dev/null
+++ b/Engine/Plugins/VirtualProduction/CaptureManager/CaptureManagerApp/Source/CaptureManagerEditor/Private/UI/TakesView.h
@@ -0,0 +1,95 @@
+// Copyright Epic Games, Inc. All Rights Reserved.
+
+#pragma once
+
+#include "Ingest/LiveLinkDeviceCapability_Ingest.h"
+
+#include "Delegates/Delegate.h"
+#include "Delegates/DelegateCombinations.h"
+#include "Input/Reply.h"
+#include "Widgets/SCompoundWidget.h"
+#include "Widgets/Views/STileView.h"
+#include "TakeVirtualAsset.h"
+#include "CaptureSourceVirtualAsset.h"
+#include "ContentBrowserDelegates.h"
+#include "IContentBrowserSingleton.h"
+
+#include "IngestManagement/IngestJobSettingsManager.h"
+
+class FCaptureManagerPanelController;
+class FCaptureSourcesView;
+class SEditableTextBox;
+class SSearchBox;
+class SSlateBoxBrush;
+
+namespace UE::CaptureManager
+{
+class SIngestJobProcessor;
+}
+
+class STakesView : public SCompoundWidget
+{
+public:
+	DECLARE_DELEGATE_TwoParams(FOnAddTakesToIngestQueue, const TArray<TStrongObjectPtr<UTakeVirtualAsset>>&, const UIngestJobSettingsBase&);
+
+	SLATE_BEGIN_ARGS(STakesView) {}
+		SLATE_ARGUMENT(FAssetPickerConfig, DevicesPickerConfig)
+		SLATE_ARGUMENT(FAssetPickerConfig, TakesPickerConfig)
+		SLATE_EVENT(FOnAddTakesToIngestQueue, OnAddTakesToIngestQueue)
+	SLATE_END_ARGS()
+
+	STakesView();
+	~STakesView() = default;
+
+	void Construct(
+		const FArguments& InArgs, 
+		TSharedRef<UE::CaptureManager::FIngestJobSettingsManager> IngestJobsSettingsManager,
+		TUniquePtr<FGetCurrentSelectionDelegate> GetCurrentSelectionDelegate
+	);
+	void SetAddToQueueButtonEnabled(bool bIsEnabled);
+
+private:
+	FReply OnAddToQueueButtonClicked();
+	TSharedRef<SWidget> OnGeneratePresetNameWidget(TSharedRef<FText> InText);
+	void OnPresetSelectionChanged(TSharedPtr<FText> InText, ESelectInfo::Type InSelectType);
+	TArray<TStrongObjectPtr<UTakeVirtualAsset>> GetSelectedTakeAssets() const;
+
+	TUniquePtr<FGetCurrentSelectionDelegate> GetCurrentSelectionDelegate;
+	TSharedPtr<SButton> AddToQueueButton;
+	FOnAddTakesToIngestQueue OnAddTakesToIngestQueue;
+
+	TArray<TSharedRef<FText>> IngestJobPresetNames;
+	TSharedPtr<UE::CaptureManager::FIngestJobSettingsManager> IngestJobSettingsManager;
+	FText CurrentPresetName;
+};
+
+class FTakesView
+{
+public:
+	FTakesView(TWeakPtr<FCaptureManagerPanelController> InController);
+
+	void CaptureDeviceStarted(FGuid InCaptureDeviceId);
+	void CaptureDeviceStopped(FGuid InCaptureDeviceId);
+	void CaptureDeviceToBeRemoved(FGuid InCaptureDeviceId);
+
+	TSharedPtr<STakesView> TakesTileView;
+private:
+	void CreateTakesView();
+	TWeakPtr<FCaptureManagerPanelController> PanelController;
+
+	TMap<FGuid, TArray<TStrongObjectPtr<UTakeVirtualAsset>>> CaptureDeviceTakesMap;
+	TArray<TStrongObjectPtr<UCaptureSourceVirtualAsset>> CaptureDevices;
+	TWeakObjectPtr<UCaptureSourceVirtualAsset> SelectedItem;
+
+	FRefreshAssetViewDelegate RefreshTakesViewDelegate;
+	FRefreshAssetViewDelegate RefreshDevicesViewDelegate;
+	FSyncToAssetsDelegate SyncToDevicesDelegate;
+
+	TArray<FAssetData> CreateTakeAssetsDataForCaptureDevice(TOptional<FString> InCaptureDeviceName = TOptional<FString>());
+	TArray<FAssetData> CreateTakeAssetsData(const TMap<FString, TArray<TStrongObjectPtr<UTakeVirtualAsset>>>& InSelectedCaptureDevicesTakesMap);
+	void AddTakesToIngestQueue(const TArray<TStrongObjectPtr<UTakeVirtualAsset>>& InTakeAssets, const UIngestJobSettingsBase& InJobPreset);
+
+	TArray<TStrongObjectPtr<UTakeVirtualAsset>> MakeTakeObjects(FGuid InCaptureDeviceId,
+																TOptional<TArray<UE::CaptureManager::FTakeId>> InTakeIds = TOptional<TArray<UE::CaptureManager::FTakeId>>());
+	TOptional<FTakeThumbnail> CreateThumbnail(const FTakeThumbnailData& InThumbnailData);
+};
diff --git a/Engine/Plugins/VirtualProduction/CaptureManager/CaptureManagerApp/Source/CaptureManagerMediaRW/Private/CaptureManagerMediaRWModule.cpp b/Engine/Plugins/VirtualProduction/CaptureManager/CaptureManagerApp/Source/CaptureManagerMediaRW/Private/CaptureManagerMediaRWModule.cpp
new file mode 100644
--- /dev/null
+++ b/Engine/Plugins/VirtualProduction/CaptureManager/CaptureManagerApp/Source/CaptureManagerMediaRW/Private/CaptureManagerMediaRWModule.cpp
@@ -0,0 +1,53 @@
+// Copyright Epic Games, Inc. All Rights Reserved.
+
+#include "CaptureManagerMediaRWModule.h"
+
+#if PLATFORM_WINDOWS && !UE_SERVER
+
+#include "Utils/WindowsRWHelpers.h"
+
+#endif
+#include "Readers/MHADepthVideoReader.h"
+#include "Readers/MHAICalibrationReader.h"
+#include "Readers/OpenCVCalibrationReader.h"
+
+#include "Writers/AudioWaveMediaWriter.h"
+#include "Writers/DepthImageWriter.h"
+#include "Writers/UnrealCalibrationWriter.h"
+
+void FCaptureManagerMediaRWModule::StartupModule()
+{
+	MediaRWManager = MakeUnique<FMediaRWManager>();
+
+	// Readers
+#if PLATFORM_WINDOWS && !UE_SERVER
+	FWindowsRWHelpers::Init();
+	FWindowsRWHelpers::RegisterReaders(MediaRWManager.Get());
+	FWindowsRWHelpers::RegisterWriters(MediaRWManager.Get());
+#endif
+
+	FMHADepthVideoReaderHelpers::RegisterReaders(MediaRWManager.Get());
+	FMHAICalibrationReaderHelpers::RegisterReaders(MediaRWManager.Get());
+	FOpenCvCalibrationReaderHelpers::RegisterReaders(MediaRWManager.Get());
+
+	// Writers
+	FAudioWaveWriterHelpers::RegisterWriters(MediaRWManager.Get());
+	FDepthExrImageWriterHelpers::RegisterWriters(MediaRWManager.Get());
+	FUnrealCalibrationWriterHelpers::RegisterWriters(MediaRWManager.Get());
+}
+
+void FCaptureManagerMediaRWModule::ShutdownModule()
+{
+#if PLATFORM_WINDOWS && !UE_SERVER
+	FWindowsRWHelpers::Deinit();
+#endif
+
+	MediaRWManager = nullptr;
+}
+
+FMediaRWManager& FCaptureManagerMediaRWModule::Get()
+{
+	return *MediaRWManager;
+}
+
+IMPLEMENT_MODULE(FCaptureManagerMediaRWModule, CaptureManagerMediaRW);
\ No newline at end of file
diff --git a/Engine/Plugins/VirtualProduction/CaptureManager/CaptureManagerApp/Source/CaptureManagerMediaRW/Private/Utils/WindowsRWHelpers.cpp b/Engine/Plugins/VirtualProduction/CaptureManager/CaptureManagerApp/Source/CaptureManagerMediaRW/Private/Utils/WindowsRWHelpers.cpp
new file mode 100644
--- /dev/null
+++ b/Engine/Plugins/VirtualProduction/CaptureManager/CaptureManagerApp/Source/CaptureManagerMediaRW/Private/Utils/WindowsRWHelpers.cpp
@@ -0,0 +1,77 @@
+// Copyright Epic Games, Inc. All Rights Reserved.
+
+#include "WindowsRWHelpers.h"
+
+#include "MediaRWManager.h"
+
+#if PLATFORM_WINDOWS && !UE_SERVER
+
+DEFINE_LOG_CATEGORY_STATIC(LogWindowsRWHelper, Log, All);
+
+bool FWindowsRWHelpers::Init()
+{
+	HRESULT Result = CoInitializeEx(nullptr, COINIT_APARTMENTTHREADED);
+
+	if (!SUCCEEDED(Result))
+	{
+		CoUninitialize();
+
+		FString Message = FormatWindowsMessage(Result);
+		UE_LOG(LogWindowsRWHelper, Error, TEXT("Failed to initialize Windows Media Foundation %s"), *Message);
+
+		return false;
+	}
+
+	Result = MFStartup(MF_VERSION);
+
+	if (!SUCCEEDED(Result))
+	{
+		CoUninitialize();
+
+		FString Message = FormatWindowsMessage(Result);
+		UE_LOG(LogWindowsRWHelper, Error, TEXT("Failed to start Windows Media Foundation %s"), *Message);
+
+		return false;
+	}
+
+	return true;
+}
+
+void FWindowsRWHelpers::Deinit()
+{
+	MFShutdown();
+	CoUninitialize();
+}
+
+void FWindowsRWHelpers::RegisterReaders(FMediaRWManager* InManager)
+{
+	TArray<FString> SupportedExtensions = { TEXT("mov"), TEXT("mp4") };
+	InManager->RegisterAudioReader(SupportedExtensions, MakeUnique<FWindowsReadersFactory>());
+	InManager->RegisterVideoReader(SupportedExtensions, MakeUnique<FWindowsReadersFactory>());
+}
+
+void FWindowsRWHelpers::RegisterWriters(FMediaRWManager* InManager)
+{
+	TArray<FString> SupportedExtensions = { TEXT("png"), TEXT("jpg"), TEXT("jpeg") };
+	InManager->RegisterImageWriter(SupportedExtensions, MakeUnique<FWindowsImageWriterFactory>());
+}
+
+FText FWindowsRWHelpers::CreateErrorMessage(HRESULT InResult, FText InMessage)
+{
+	FText WindowsErrorMessage = FText::FromString(FWindowsRWHelpers::FormatWindowsMessage(InResult));
+	FText ErrorMessage = FText::Format(FText::FromString(TEXT("{0}: {1}")), MoveTemp(InMessage), WindowsErrorMessage);
+
+	return ErrorMessage;
+}
+
+FString FWindowsRWHelpers::FormatWindowsMessage(HRESULT InResult)
+{
+	constexpr uint32 BufSize = 1024;
+	WIDECHAR Buffer[BufSize];
+
+	FormatMessage(FORMAT_MESSAGE_FROM_SYSTEM, nullptr, InResult, 0, Buffer, BufSize, nullptr);
+
+	return Buffer;
+}
+
+#endif // PLATFORM_WINDOWS && !UE_SERVER
\ No newline at end of file
diff --git a/Engine/Plugins/VirtualProduction/CaptureManager/CaptureManagerApp/Source/CaptureManagerMediaRW/Private/Writers/WindowsImageWriter.cpp b/Engine/Plugins/VirtualProduction/CaptureManager/CaptureManagerApp/Source/CaptureManagerMediaRW/Private/Writers/WindowsImageWriter.cpp
new file mode 100644
--- /dev/null
+++ b/Engine/Plugins/VirtualProduction/CaptureManager/CaptureManagerApp/Source/CaptureManagerMediaRW/Private/Writers/WindowsImageWriter.cpp
@@ -0,0 +1,646 @@
+// Copyright Epic Games, Inc. All Rights Reserved.
+
+#include "WindowsImageWriter.h"
+
+#include "MediaRWManager.h"
+
+#include "IImageWrapperModule.h"
+#include "Modules/ModuleManager.h"
+#include "Async/ParallelFor.h"
+
+#if PLATFORM_WINDOWS && !UE_SERVER
+
+#include "Utils/WindowsRWHelpers.h"
+
+#define LOCTEXT_NAMESPACE "WindowsImageReader"
+
+DEFINE_LOG_CATEGORY_STATIC(LogWindowsImageWriter, Log, All);
+
+static const GUID UnsupportedEncoder = GUID_NULL;
+
+#define WINIW_CHECK_AND_RETURN_ERROR_MESSAGE(Result, Message)                        \
+	if (FAILED(Result))                                                              \
+	{                                                                                \
+		FText ErrorMessage = FWindowsRWHelpers::CreateErrorMessage(Result, Message); \
+		UE_LOG(LogWindowsImageWriter, Error, TEXT("%s"), *ErrorMessage.ToString());  \
+		return ErrorMessage;														 \
+	}
+
+TUniquePtr<IImageWriter> FWindowsImageWriterFactory::CreateImageWriter()
+{
+	return MakeUnique<FWindowsImageWriter>();
+}
+
+namespace UE::CaptureManager::Private
+{
+
+WICBitmapTransformOptions ConvertToWindowsTransformOptions(EMediaOrientation InOrientation)
+{
+	switch (InOrientation)
+	{
+		case EMediaOrientation::CW90:
+			return WICBitmapTransformRotate270;
+		case EMediaOrientation::CW180:
+			return WICBitmapTransformRotate180;
+		case EMediaOrientation::CW270:
+			return WICBitmapTransformRotate90;
+		case EMediaOrientation::Original:
+		default:
+			return WICBitmapTransformRotate0;
+	}
+}
+
+FIntPoint GetNewDimensions(const FMediaTextureSample* InNewSample)
+{
+	switch (InNewSample->Orientation)
+	{
+		case EMediaOrientation::CW90:
+		case EMediaOrientation::CW270:
+			return FIntPoint(InNewSample->Dimensions.Y, InNewSample->Dimensions.X);
+		case EMediaOrientation::Original:
+		case EMediaOrientation::CW180:
+		default:
+			return InNewSample->Dimensions;
+	}
+}
+
+WICPixelFormatGUID ConvertPixelFormat(EMediaTexturePixelFormat InPixelFormat)
+{
+	switch (InPixelFormat)
+	{
+		case EMediaTexturePixelFormat::U8_RGBA:
+			return GUID_WICPixelFormat32bppRGBA;
+		case EMediaTexturePixelFormat::U8_RGB:
+			return GUID_WICPixelFormat24bppRGB;
+		case EMediaTexturePixelFormat::U8_BGRA:
+			return GUID_WICPixelFormat32bppBGRA;
+		case EMediaTexturePixelFormat::U8_BGR:
+			return GUID_WICPixelFormat24bppBGR;
+		case EMediaTexturePixelFormat::U8_Mono:
+			return GUID_WICPixelFormat8bppGray;
+		case EMediaTexturePixelFormat::U16_Mono:
+			return GUID_WICPixelFormat16bppGray;
+		case EMediaTexturePixelFormat::F_Mono:
+			return GUID_WICPixelFormat32bppGrayFloat;
+		case EMediaTexturePixelFormat::U8_I420:
+		case EMediaTexturePixelFormat::U8_YUY2:
+		case EMediaTexturePixelFormat::U8_NV12:
+		case EMediaTexturePixelFormat::U8_I444:
+		default:
+			return GUID_NULL;
+	}
+}
+
+FString GetPixelFormatString(EMediaTexturePixelFormat InPixelFormat)
+{
+	switch (InPixelFormat)
+	{
+		case EMediaTexturePixelFormat::U8_RGBA:
+			return TEXT("U8 RGBA");
+		case EMediaTexturePixelFormat::U8_RGB:
+			return TEXT("U8 RGBA");
+		case EMediaTexturePixelFormat::U8_BGRA:
+			return TEXT("U8 RGBA");
+		case EMediaTexturePixelFormat::U8_BGR:
+			return TEXT("U8 RGBA");
+		case EMediaTexturePixelFormat::U8_Mono:
+			return TEXT("U8 RGBA");
+		case EMediaTexturePixelFormat::U16_Mono:
+			return TEXT("U8 RGBA");
+		case EMediaTexturePixelFormat::F_Mono:
+			return TEXT("U8 RGBA");
+		case EMediaTexturePixelFormat::U8_I420:
+			return TEXT("U8 I420");
+		case EMediaTexturePixelFormat::U8_YUY2:
+			return TEXT("U8 YUY2");
+		case EMediaTexturePixelFormat::U8_NV12:
+			return TEXT("U8 NV12");
+		case EMediaTexturePixelFormat::U8_I444:
+			return TEXT("U8 I444");
+		default:
+			return TEXT("Undefined");
+	}
+}
+
+template<typename T>
+T ScalePixel(T InPixel, TRange<T> InInputRange, TRange<T> InOutputRange)
+{
+	const float Factor = static_cast<float>(InOutputRange.GetUpperBoundValue()) / (InInputRange.GetUpperBoundValue() - InInputRange.GetLowerBoundValue());
+	const uint64 NewPixelValue = Factor * (InPixel - InInputRange.GetLowerBoundValue()) + 0.5;
+	return FMath::Clamp(NewPixelValue, InOutputRange.GetLowerBoundValue(), InOutputRange.GetUpperBoundValue());
+}
+
+static inline void RGBfromYUV(uint8& R, uint8& G, uint8& B, uint8 Y, uint8 U, uint8 V)
+{
+	int64 BigY = Y - 16;
+	int64 BigU = U - 128;
+	int64 BigV = V - 128;
+
+	int64 BigR = (298 * BigY + 409 * BigV + 128) >> 8;
+	int64 BigG = (298 * BigY - 100 * BigU - 208 * BigV + 128) >> 8;
+	int64 BigB = (298 * BigY + 516 * BigU + 128) >> 8;
+
+	R = FMath::Clamp(BigR, 0, 255);
+	G = FMath::Clamp(BigG, 0, 255);
+	B = FMath::Clamp(BigB, 0, 255);
+}
+
+class FConverter
+{
+public:
+
+	static FConverter Create(TComPtr<IWICImagingFactory> InFactory,
+							 TComPtr<IWICBitmap> InBitmap,
+							 EMediaTexturePixelFormat InInputFormat,
+							 EMediaTexturePixelFormat InOutputFormat)
+	{
+		FConverter Converter;
+
+		Converter.PixelFormat = Private::ConvertPixelFormat(InInputFormat);
+		Converter.Source = InBitmap;
+
+		if (InInputFormat == InOutputFormat)
+		{
+			return Converter;
+		}
+
+		TComPtr<IWICFormatConverter> WICConverter;
+		HRESULT Result = InFactory->CreateFormatConverter(&WICConverter);
+
+		if (SUCCEEDED(Result))
+		{
+			WICPixelFormatGUID OutPixelFormat = Private::ConvertPixelFormat(InOutputFormat);
+			if (OutPixelFormat != GUID_NULL)
+			{
+				BOOL bCanConvert = false;
+				Result = WICConverter->CanConvert(Converter.PixelFormat, OutPixelFormat, &bCanConvert);
+				if (SUCCEEDED(Result) && bCanConvert)
+				{
+					Result = WICConverter->Initialize(InBitmap, OutPixelFormat, WICBitmapDitherTypeNone, nullptr, 0.f, WICBitmapPaletteTypeCustom);
+					if (SUCCEEDED(Result))
+					{
+						Converter.PixelFormat = OutPixelFormat;
+						Converter.Source = WICConverter;
+					}
+					else
+					{
+						FString Message =
+							FString::Format(TEXT("Failed to convert from ({0}) to the ({1}) pixel format"),
+											{ Private::GetPixelFormatString(InInputFormat), Private::GetPixelFormatString(InOutputFormat) });
+						UE_LOG(LogWindowsImageWriter, Warning, TEXT("%s"), *Message);
+					}
+				}
+				else
+				{
+					FString Message =
+						FString::Format(TEXT("Unsupported pixel format conversion: {0} -> {1}"),
+										{ Private::GetPixelFormatString(InInputFormat), Private::GetPixelFormatString(InOutputFormat) });
+
+					UE_LOG(LogWindowsImageWriter, Warning, TEXT("%s"), *Message);
+				}
+			}
+			else
+			{
+				UE_LOG(LogWindowsImageWriter, Warning, TEXT("Unsupported pixel format provided."));
+			}
+		}
+		else
+		{
+			UE_LOG(LogWindowsImageWriter, Warning, TEXT("Failed to create converter from WIC"));
+		}
+
+		return Converter;
+	}
+
+	TComPtr<IWICBitmapSource> GetSource()
+	{
+		return Source;
+	}
+
+	WICPixelFormatGUID GetPixelFormat()
+	{
+		return PixelFormat;
+	}
+
+private:
+
+	TComPtr<IWICBitmapSource> Source;
+	WICPixelFormatGUID PixelFormat;
+};
+
+}
+
+FWindowsImageWriter::FWindowsImageWriter() = default;
+FWindowsImageWriter::~FWindowsImageWriter() = default;
+
+TOptional<FText> FWindowsImageWriter::Open(const FString& InDirectory, const FString& InFileName, const FString& InFormat)
+{
+	HRESULT Result = CoCreateInstance(CLSID_WICImagingFactory, NULL, CLSCTX_INPROC_SERVER, IID_PPV_ARGS(&WindowsImagingFactory));
+	WINIW_CHECK_AND_RETURN_ERROR_MESSAGE(Result, LOCTEXT("Open_FailedToInitialize", "Failed to initialize Windows Imaging Component"));
+
+	GUID EncoderGuid = GetEncoderGuidBasedOnFormat(InFormat);
+	if (EncoderGuid == UnsupportedEncoder)
+	{
+		FText ErrorMessage = LOCTEXT("Open_UnsupportedEncoder", "Unsupported encoder provided");
+		UE_LOG(LogWindowsImageWriter, Error, TEXT("%s"), *ErrorMessage.ToString());
+
+		return ErrorMessage;
+	}
+	
+	Directory = InDirectory;
+	FileName = InFileName;
+	Format = InFormat;
+
+	return {};
+}
+
+TOptional<FText> FWindowsImageWriter::Close()
+{
+	WindowsImagingFactory = nullptr;
+
+	return {};
+}
+
+TOptional<FText> FWindowsImageWriter::Append(UE::CaptureManager::FMediaTextureSample* InSample)
+{
+	using namespace UE::CaptureManager;
+
+	if (InSample->DesiredFormat == EMediaTexturePixelFormat::U8_Mono || 
+		InSample->DesiredFormat == EMediaTexturePixelFormat::Undefined)
+	{
+		InSample->DesiredFormat = EMediaTexturePixelFormat::U8_Mono;
+		if (InSample->CurrentFormat == UE::CaptureManager::EMediaTexturePixelFormat::U8_YUY2)
+		{
+			InSample->Buffer = ConvertYUY2ToMono(InSample, false);
+			InSample->CurrentFormat = EMediaTexturePixelFormat::U8_Mono;
+		}
+		else if (InSample->CurrentFormat == UE::CaptureManager::EMediaTexturePixelFormat::U8_I420 ||
+				 InSample->CurrentFormat == UE::CaptureManager::EMediaTexturePixelFormat::U8_NV12)
+		{
+			InSample->Buffer = ConvertYUVToMono(InSample);
+			InSample->CurrentFormat = EMediaTexturePixelFormat::U8_Mono;
+		}
+		else
+		{
+			// Do nothing
+		}
+	}
+	else if (InSample->DesiredFormat == EMediaTexturePixelFormat::U8_BGRA)
+	{
+		if (InSample->CurrentFormat == UE::CaptureManager::EMediaTexturePixelFormat::U8_I420)
+		{
+			InSample->Buffer = ConvertI420ToBGRA(InSample);
+			InSample->CurrentFormat = EMediaTexturePixelFormat::U8_BGRA;
+		}
+		else if (InSample->CurrentFormat == UE::CaptureManager::EMediaTexturePixelFormat::U8_NV12)
+		{
+			InSample->Buffer = ConvertNV12ToBGRA(InSample);
+			InSample->CurrentFormat = EMediaTexturePixelFormat::U8_BGRA;
+		}
+		else if (InSample->CurrentFormat == UE::CaptureManager::EMediaTexturePixelFormat::U8_YUY2)
+		{
+			InSample->Buffer = ConvertYUY2ToBGRA(InSample);
+			InSample->CurrentFormat = EMediaTexturePixelFormat::U8_BGRA;
+		}
+		else
+		{
+			// Do nothing
+		}
+	}
+
+	WICPixelFormatGUID PixelFormat = Private::ConvertPixelFormat(InSample->CurrentFormat);
+	if (PixelFormat == GUID_NULL)
+	{
+		FText ErrorMessage = 
+			FText::Format(
+				LOCTEXT("Append_UnsupportedPixelFormat", "Image has unsupported pixel format ({0})"), 
+						FText::FromString(*Private::GetPixelFormatString(InSample->CurrentFormat)));
+		UE_LOG(LogWindowsImageWriter, Error, TEXT("%s"), *ErrorMessage.ToString());
+
+		return ErrorMessage;
+	}
+
+	TComPtr<IWICBitmap> Bitmap;
+	HRESULT Result = WindowsImagingFactory->CreateBitmapFromMemory(InSample->Dimensions.X,
+																   InSample->Dimensions.Y,
+																   PixelFormat,
+																   InSample->Stride * GetNumberOfChannels(InSample->CurrentFormat),
+																   InSample->Buffer.Num(),
+																   InSample->Buffer.GetData(),
+																   &Bitmap);
+	WINIW_CHECK_AND_RETURN_ERROR_MESSAGE(Result, LOCTEXT("Append_FailedToCreateBitmap", "Failed to create bitmap from memory"));
+
+	Private::FConverter Converter = Private::FConverter::Create(WindowsImagingFactory, Bitmap, InSample->CurrentFormat, InSample->DesiredFormat);
+
+	TComPtr<IWICBitmapSource> Source = Converter.GetSource();
+	PixelFormat = Converter.GetPixelFormat();
+
+	UE_LOG(LogWindowsImageWriter, Verbose, TEXT("Output pixel format is %s"), *Private::GetPixelFormatString(InSample->CurrentFormat));
+	
+	TComPtr<IWICBitmapFlipRotator> Rotator = nullptr;
+
+	GUID EncoderGuid = GetEncoderGuidBasedOnFormat(Format);
+	if (InSample->Orientation != EMediaOrientation::Original && EncoderGuid != GUID_ContainerFormatJpeg)
+	{
+		Result = WindowsImagingFactory->CreateBitmapFlipRotator(&Rotator);
+		WINIW_CHECK_AND_RETURN_ERROR_MESSAGE(Result, LOCTEXT("Open_FailedToCreateRotator", "Failed to create Rotator object"));
+
+		Result = Rotator->Initialize(Source, Private::ConvertToWindowsTransformOptions(InSample->Orientation));
+		WINIW_CHECK_AND_RETURN_ERROR_MESSAGE(Result, LOCTEXT("Append_FailedToInitializeRotator", "Failed to initialize rotator from memory"));
+
+		Source = Rotator;
+	}
+	
+	FString FrameFileName = CreateFrameFileName();
+	FString Path = FPaths::SetExtension(Directory / FrameFileName, Format);
+
+	TComPtr<IWICStream> Stream;
+	Result = WindowsImagingFactory->CreateStream(&Stream);
+	WINIW_CHECK_AND_RETURN_ERROR_MESSAGE(Result, LOCTEXT("Append_FailedToCreateStream", "Failed to create stream object used for writing the frame"));
+
+	Result = Stream->InitializeFromFilename(*Path, GENERIC_WRITE);
+	WINIW_CHECK_AND_RETURN_ERROR_MESSAGE(Result, FText::Format(LOCTEXT("Append_FailedToInitializeStream", "Failed to initialize stream object based of the filename: {0}"),
+														 FText::FromString(Path)));
+	TComPtr<IWICBitmapEncoder> Encoder;
+
+	Result = WindowsImagingFactory->CreateEncoder(EncoderGuid, nullptr, &Encoder);
+	WINIW_CHECK_AND_RETURN_ERROR_MESSAGE(Result, FText::Format(LOCTEXT("Open_FailedToCreateEncoder", "Failed to create Encoder object for format {0}"),
+														 FText::FromString(Format)));
+
+	Result = Encoder->Initialize(Stream, WICBitmapEncoderNoCache);
+	WINIW_CHECK_AND_RETURN_ERROR_MESSAGE(Result, FText::Format(LOCTEXT("Append_FailedToInitializeEncoder", "Failed to initialize encoder for format: {0}"), FText::FromString(Format)));
+
+	IPropertyBag2* PropertyBag = nullptr;
+	TComPtr<IWICBitmapFrameEncode> EncodedFrame;
+	Result = Encoder->CreateNewFrame(&EncodedFrame, &PropertyBag);
+	WINIW_CHECK_AND_RETURN_ERROR_MESSAGE(Result, LOCTEXT("Append_FailedToCreateEncodedFrame", "Failed to create encoded frame"));
+
+	if (EncoderGuid == GUID_ContainerFormatJpeg)
+	{
+		FString PropertyName = TEXT("BitmapTransform");
+		PROPBAG2 Options = { 0 };
+		Options.pstrName = PropertyName.GetCharArray().GetData();
+		VARIANT Variant;
+		VariantInit(&Variant);
+		Variant.vt = VT_UI1;
+		Variant.bVal = Private::ConvertToWindowsTransformOptions(InSample->Orientation);
+		constexpr ULONG NumberOfProperties = 1;
+		PropertyBag->Write(NumberOfProperties, &Options, &Variant);
+	}
+
+	Result = EncodedFrame->Initialize(PropertyBag);
+	WINIW_CHECK_AND_RETURN_ERROR_MESSAGE(Result, LOCTEXT("Append_FailedToInitializeEncodedFrame", "Failed to initialize encoded frame"));
+
+	FIntPoint NewDimensions = Private::GetNewDimensions(InSample);
+	Result = EncodedFrame->SetSize(NewDimensions.X, NewDimensions.Y);
+
+	Result = EncodedFrame->SetPixelFormat(&PixelFormat);
+	WINIW_CHECK_AND_RETURN_ERROR_MESSAGE(Result, LOCTEXT("Append_FailedToSetPixelFormat", "Failed to configure pixel format on the encoded frame"));
+
+	Result = EncodedFrame->WriteSource(Source, nullptr);
+	WINIW_CHECK_AND_RETURN_ERROR_MESSAGE(Result, LOCTEXT("Append_FailedToWriteData", "Failed to write the data to the encoded frame"));
+
+	Result = EncodedFrame->Commit();
+	WINIW_CHECK_AND_RETURN_ERROR_MESSAGE(Result, LOCTEXT("Append_FailedToFinalizeConfiguring", "Failed to finalize configuring encoded frame"));
+
+	Result = Encoder->Commit();
+	WINIW_CHECK_AND_RETURN_ERROR_MESSAGE(Result, LOCTEXT("Append_FailedToFinalize", "Failed to finalize encoding"));
+
+	++FrameNumber;
+
+	return {};
+}
+
+GUID FWindowsImageWriter::GetEncoderGuidBasedOnFormat(const FString& InFormat)
+{
+	IImageWrapperModule& ImageWrapperModule = FModuleManager::LoadModuleChecked<IImageWrapperModule>(FName("ImageWrapper"));
+	EImageFormat ImageFormat = ImageWrapperModule.GetImageFormatFromExtension(*InFormat);
+
+	if (ImageFormat == EImageFormat::JPEG)
+	{
+		return GUID_ContainerFormatJpeg;
+	}
+	else if (ImageFormat == EImageFormat::PNG)
+	{
+		return GUID_ContainerFormatPng;
+	}
+
+	return UnsupportedEncoder;
+}
+
+TArray<uint8> FWindowsImageWriter::ConvertYUVToMono(const UE::CaptureManager::FMediaTextureSample* InSample, bool InScaleRange)
+{
+	using namespace UE::CaptureManager;
+
+	TArray<uint8> NewBuffer;
+	NewBuffer.SetNumUninitialized(InSample->Stride * InSample->Dimensions.Y);
+
+	uint32 StrideY = InSample->Stride;
+	int32 MonoStride = InSample->Stride;
+
+	ParallelFor(InSample->Dimensions.Y, [&NewBuffer, StrideY, MonoStride, InScaleRange, InSample](int32 Y) mutable
+	{
+		for (int32 X = 0; X < InSample->Dimensions.X; ++X)
+		{
+			int32 YIndex = Y * StrideY + X;
+			uint8 PixelValue = InSample->Buffer[YIndex];
+
+			if (InScaleRange)
+			{
+				using FRange = TRange<uint8>;
+				PixelValue = Private::ScalePixel(PixelValue,
+												 FRange::Inclusive(16, 235),
+												 FRange::Inclusive(MIN_uint8, MAX_uint8));
+			}
+
+			NewBuffer[Y * MonoStride + X] = PixelValue;
+		}
+	});
+
+	return NewBuffer;
+}
+
+TArray<uint8> FWindowsImageWriter::ConvertYUY2ToMono(const UE::CaptureManager::FMediaTextureSample* InSample, bool InScaleRange)
+{
+	using namespace UE::CaptureManager;
+
+	TArray<uint8> NewBuffer;
+	NewBuffer.SetNumUninitialized(InSample->Stride * InSample->Dimensions.Y);
+
+	int32 YUY2Stride = InSample->Stride * 2;
+	int32 MonoStride = InSample->Stride;
+
+	ParallelFor(InSample->Dimensions.Y, [&NewBuffer, YUY2Stride, MonoStride, InScaleRange, InSample](int32 Y) mutable
+	{
+		for (int32 X = 0; X < InSample->Dimensions.X; X += 2)
+		{
+			int32 YIndex = Y * YUY2Stride + X * 2;
+
+			uint8_t Y0 = InSample->Buffer[YIndex];
+			uint8_t Y1 = InSample->Buffer[YIndex + 2];
+
+			if (InScaleRange)
+			{
+				using FRange = TRange<uint8>;
+				Y0 = Private::ScalePixel(Y0,
+										 FRange::Inclusive(16, 235),
+										 FRange::Inclusive(MIN_uint8, MAX_uint8));
+				Y1 = Private::ScalePixel(Y1,
+										 FRange::Inclusive(16, 235),
+										 FRange::Inclusive(MIN_uint8, MAX_uint8));
+
+			}
+
+			NewBuffer[Y * MonoStride + X] = Y0;
+			NewBuffer[Y * MonoStride + X + 1] = Y1;
+		}
+	});
+
+	return NewBuffer;
+}
+
+TArray<uint8> FWindowsImageWriter::ConvertI420ToBGRA(const UE::CaptureManager::FMediaTextureSample* InSample)
+{
+	using namespace UE::CaptureManager;
+
+	EMediaTexturePixelFormat OutputFormat = EMediaTexturePixelFormat::U8_BGRA;
+	uint32 Channels = GetNumberOfChannels(OutputFormat);
+
+	TArray<uint8> NewBuffer;
+	NewBuffer.SetNumUninitialized(InSample->Stride * InSample->Dimensions.Y * Channels);
+
+	uint32 StrideY = InSample->Stride;
+	uint32 StrideUV = InSample->Stride / 2;
+	int32 BGRAStride = InSample->Stride * Channels;
+
+	uint32 YPlaneSize = StrideY * InSample->Dimensions.Y;
+	uint32 UPlaneSize = StrideUV * (InSample->Dimensions.Y / 2);
+	uint32 VPlaneSize = InSample->Buffer.Num() - (YPlaneSize + UPlaneSize);
+
+	const TArrayView<const uint8> YPlane = MakeArrayView(InSample->Buffer.GetData(), YPlaneSize);
+	const TArrayView<const uint8> UPlane = MakeArrayView(InSample->Buffer.GetData() + YPlaneSize, UPlaneSize);
+	const TArrayView<const uint8> VPlane = MakeArrayView(InSample->Buffer.GetData() + YPlaneSize + UPlaneSize, VPlaneSize);
+
+	ParallelFor(InSample->Dimensions.Y, [&NewBuffer, YPlane, UPlane, VPlane, StrideY, StrideUV, BGRAStride, InSample](int32 Y) mutable
+	{
+		for (int32 X = 0; X < InSample->Dimensions.X; ++X)
+		{
+			int32 YIndex = Y * StrideY + X;
+			int32 UVIndex = (Y / 2) * StrideUV + (X / 2);
+
+			uint8_t Y0 = YPlane[YIndex];
+			uint8_t U = UPlane[UVIndex];
+			uint8_t V = VPlane[UVIndex];
+
+			uint8 R, G, B;
+			Private::RGBfromYUV(R, G, B, Y0, U, V);
+
+			NewBuffer[Y * BGRAStride + X * 4] = B;
+			NewBuffer[Y * BGRAStride + X * 4 + 1] = G;
+			NewBuffer[Y * BGRAStride + X * 4 + 2] = R;
+			NewBuffer[Y * BGRAStride + X * 4 + 3] = 255;
+		}
+	});
+
+	return NewBuffer;
+}
+
+TArray<uint8> FWindowsImageWriter::ConvertNV12ToBGRA(const UE::CaptureManager::FMediaTextureSample* InSample)
+{
+	using namespace UE::CaptureManager;
+
+	EMediaTexturePixelFormat OutputFormat = EMediaTexturePixelFormat::U8_BGRA;
+	uint32 Channels = GetNumberOfChannels(OutputFormat);
+
+	TArray<uint8> NewBuffer;
+	NewBuffer.SetNumUninitialized(InSample->Stride * InSample->Dimensions.Y * Channels);
+
+	uint32 StrideY = InSample->Stride;
+	int32 BGRAStride = InSample->Stride * Channels;
+
+	uint32 YPlaneSize = StrideY * InSample->Dimensions.Y;
+	uint32 UVPlaneSize = InSample->Buffer.Num() - YPlaneSize;
+
+	const TArrayView<const uint8> YPlane = MakeArrayView(InSample->Buffer.GetData(), YPlaneSize);
+	const TArrayView<const uint8> UVPlane = MakeArrayView(InSample->Buffer.GetData() + YPlaneSize, UVPlaneSize);
+
+	ParallelFor(InSample->Dimensions.Y, [&NewBuffer, YPlane, UVPlane, StrideY, BGRAStride, InSample, Channels](int32 Y) mutable
+	{
+		for (int32 X = 0; X < InSample->Dimensions.X; ++X)
+		{
+			int32 YIndex = Y * StrideY + X;
+			int32 UVIndex = (Y / 2) * StrideY + (X & ~1);
+
+			uint8_t Y0 = YPlane[YIndex];
+			uint8_t U = UVPlane[UVIndex];
+			uint8_t V = UVPlane[UVIndex + 1];
+
+			uint8 R, G, B;
+			Private::RGBfromYUV(R, G, B, Y0, U, V);
+
+			NewBuffer[Y * BGRAStride + X * 4] = B;
+			NewBuffer[Y * BGRAStride + X * 4 + 1] = G;
+			NewBuffer[Y * BGRAStride + X * 4 + 2] = R;
+			NewBuffer[Y * BGRAStride + X * 4 + 3] = 255;
+		}	
+	});
+
+	return NewBuffer;
+}
+
+TArray<uint8> FWindowsImageWriter::ConvertYUY2ToBGRA(const UE::CaptureManager::FMediaTextureSample* InSample)
+{
+	using namespace UE::CaptureManager;
+
+	EMediaTexturePixelFormat OutputFormat = EMediaTexturePixelFormat::U8_BGRA;
+
+	uint32 Channels = GetNumberOfChannels(OutputFormat);
+
+	TArray<uint8> NewBuffer;
+	NewBuffer.SetNumUninitialized(InSample->Stride * InSample->Dimensions.Y * Channels);
+
+	int32 YUY2Stride = InSample->Stride * 2;
+	int32 BGRAStride = InSample->Stride * Channels;
+
+	ParallelFor(InSample->Dimensions.Y, [&NewBuffer, YUY2Stride, BGRAStride, InSample](int32 Y) mutable
+	{
+		for (int32 X = 0; X < InSample->Dimensions.X; X += 2)
+		{
+			int32 YIndex = Y * YUY2Stride + X * 2;
+
+			uint8 Y0 = InSample->Buffer[YIndex];
+			uint8 U = InSample->Buffer[YIndex + 1];
+			uint8 Y1 = InSample->Buffer[YIndex + 2];
+			uint8 V = InSample->Buffer[YIndex + 3];
+
+			uint8 R0, G0, B0;
+			Private::RGBfromYUV(R0, G0, B0, Y0, U, V);
+
+			uint8 R1, G1, B1;
+			Private::RGBfromYUV(R1, G1, B1, Y1, U, V);
+
+			NewBuffer[Y * BGRAStride + X * 4] = B0;
+			NewBuffer[Y * BGRAStride + X * 4 + 1] = G0;
+			NewBuffer[Y * BGRAStride + X * 4 + 2] = R0;
+			NewBuffer[Y * BGRAStride + X * 4 + 3] = 255;
+
+			NewBuffer[Y * BGRAStride + (X + 1) * 4] = B1;
+			NewBuffer[Y * BGRAStride + (X + 1) * 4 + 1] = G1;
+			NewBuffer[Y * BGRAStride + (X + 1) * 4 + 2] = R1;
+			NewBuffer[Y * BGRAStride + (X + 1) * 4 + 3] = 255;
+		}
+	});
+
+	return NewBuffer;
+}
+
+FString FWindowsImageWriter::CreateFrameFileName()
+{
+	FString FrameIndexStr = FString::Printf(TEXT("%06d"), FrameNumber);
+
+	return FString::Format(TEXT("{0}_{1}"), { FileName, FrameIndexStr });
+}
+
+#undef LOCTEXT_NAMESPACE
+
+#endif // PLATFORM_WINDOWS && !UE_SERVER
diff --git a/Engine/Plugins/VirtualProduction/CaptureManager/CaptureManagerApp/Source/CaptureManagerMediaRW/Private/Writers/WindowsImageWriter.h b/Engine/Plugins/VirtualProduction/CaptureManager/CaptureManagerApp/Source/CaptureManagerMediaRW/Private/Writers/WindowsImageWriter.h
new file mode 100644
--- /dev/null
+++ b/Engine/Plugins/VirtualProduction/CaptureManager/CaptureManagerApp/Source/CaptureManagerMediaRW/Private/Writers/WindowsImageWriter.h
@@ -0,0 +1,68 @@
+// Copyright Epic Games, Inc. All Rights Reserved.
+
+#pragma once
+
+#include "Misc/Build.h"
+#include "HAL/Platform.h"
+
+#if PLATFORM_WINDOWS && !UE_SERVER
+
+#include "Windows/WindowsHWrapper.h"
+#include "Windows/AllowWindowsPlatformTypes.h"
+
+THIRD_PARTY_INCLUDES_START
+
+#include <windows.h>
+#include <mfapi.h>
+#include <mfidl.h>
+#include <wincodec.h>
+#include "Microsoft/COMPointer.h"
+
+THIRD_PARTY_INCLUDES_END
+
+#include "Windows/HideWindowsPlatformTypes.h"
+
+#include "IMediaWriter.h"
+#include "IMediaRWFactory.h"
+
+class FWindowsImageWriterFactory final : public IImageWriterFactory
+{
+public:
+
+	virtual TUniquePtr<IImageWriter> CreateImageWriter() override;
+};
+
+class FWindowsImageWriter final : public IImageWriter
+{
+public:
+
+	FWindowsImageWriter();
+	virtual ~FWindowsImageWriter() override;
+
+	virtual TOptional<FText> Open(const FString& InDirectory, const FString& InFileName, const FString& InFormat) override;
+	virtual TOptional<FText> Close() override;
+
+	virtual TOptional<FText> Append(UE::CaptureManager::FMediaTextureSample* InSample) override;
+
+private:
+
+	GUID GetEncoderGuidBasedOnFormat(const FString& InFormat);
+
+	// Scaling from 16-235 range to 0-255 range
+	TArray<uint8> ConvertYUVToMono(const UE::CaptureManager::FMediaTextureSample* InSample, bool InScaleRange = true); // Any YUV to Mono
+	TArray<uint8> ConvertYUY2ToMono(const UE::CaptureManager::FMediaTextureSample* InSample, bool InScaleRange = true);
+	TArray<uint8> ConvertI420ToBGRA(const UE::CaptureManager::FMediaTextureSample* InSample);
+	TArray<uint8> ConvertNV12ToBGRA(const UE::CaptureManager::FMediaTextureSample* InSample);
+	TArray<uint8> ConvertYUY2ToBGRA(const UE::CaptureManager::FMediaTextureSample* InSample);
+	FString CreateFrameFileName();
+
+	TComPtr<IWICImagingFactory> WindowsImagingFactory;
+
+	FString Directory;
+	FString FileName;
+	FString Format;
+
+	uint32 FrameNumber = 0;
+};
+
+#endif // PLATFORM_WINDOWS && !UE_SERVER
\ No newline at end of file
diff --git a/Engine/Plugins/VirtualProduction/CaptureManager/CaptureManagerApp/Source/CaptureManagerMediaRW/Public/CaptureManagerMediaRWModule.h b/Engine/Plugins/VirtualProduction/CaptureManager/CaptureManagerApp/Source/CaptureManagerMediaRW/Public/CaptureManagerMediaRWModule.h
new file mode 100644
--- /dev/null
+++ b/Engine/Plugins/VirtualProduction/CaptureManager/CaptureManagerApp/Source/CaptureManagerMediaRW/Public/CaptureManagerMediaRWModule.h
@@ -0,0 +1,21 @@
+// Copyright Epic Games, Inc. All Rights Reserved.
+
+#pragma once
+
+#include "Modules/ModuleManager.h"
+
+#include "MediaRWManager.h"
+
+class CAPTUREMANAGERMEDIARW_API FCaptureManagerMediaRWModule : public IModuleInterface
+{
+public:
+
+	void StartupModule();
+	void ShutdownModule();
+
+	FMediaRWManager& Get();
+
+private:
+
+	TUniquePtr<FMediaRWManager> MediaRWManager;
+};
\ No newline at end of file
diff --git a/Engine/Plugins/VirtualProduction/CaptureManager/CaptureManagerApp/Source/CaptureManagerMediaRW/Public/IMediaRWFactory.h b/Engine/Plugins/VirtualProduction/CaptureManager/CaptureManagerApp/Source/CaptureManagerMediaRW/Public/IMediaRWFactory.h
new file mode 100644
--- /dev/null
+++ b/Engine/Plugins/VirtualProduction/CaptureManager/CaptureManagerApp/Source/CaptureManagerMediaRW/Public/IMediaRWFactory.h
@@ -0,0 +1,60 @@
+// Copyright Epic Games, Inc. All Rights Reserved.
+
+#pragma once
+
+#include "IMediaReader.h"
+#include "IMediaWriter.h"
+
+class CAPTUREMANAGERMEDIARW_API IAudioReaderFactory
+{
+public:
+
+	virtual ~IAudioReaderFactory() = default;
+
+	virtual TUniquePtr<IAudioReader> CreateAudioReader() = 0;
+};
+
+class CAPTUREMANAGERMEDIARW_API IVideoReaderFactory
+{
+public:
+
+	virtual ~IVideoReaderFactory() = default;
+
+	virtual TUniquePtr<IVideoReader> CreateVideoReader() = 0;
+};
+
+class CAPTUREMANAGERMEDIARW_API ICalibrationReaderFactory
+{
+public:
+
+	virtual ~ICalibrationReaderFactory() = default;
+
+	virtual TUniquePtr<ICalibrationReader> CreateCalibrationReader() = 0;
+};
+
+class CAPTUREMANAGERMEDIARW_API IAudioWriterFactory
+{
+public:
+
+	virtual ~IAudioWriterFactory() = default;
+
+	virtual TUniquePtr<IAudioWriter> CreateAudioWriter() = 0;
+};
+
+class CAPTUREMANAGERMEDIARW_API IImageWriterFactory
+{
+public:
+
+	virtual ~IImageWriterFactory() = default;
+
+	virtual TUniquePtr<IImageWriter> CreateImageWriter() = 0;
+};
+
+class CAPTUREMANAGERMEDIARW_API ICalibrationWriterFactory
+{
+public:
+
+	virtual ~ICalibrationWriterFactory() = default;
+
+	virtual TUniquePtr<ICalibrationWriter> CreateCalibrationWriter() = 0;
+};
\ No newline at end of file
diff --git a/Engine/Plugins/VirtualProduction/CaptureManager/CaptureManagerApp/Source/CaptureManagerMediaRW/Public/MediaRWManager.h b/Engine/Plugins/VirtualProduction/CaptureManager/CaptureManagerApp/Source/CaptureManagerMediaRW/Public/MediaRWManager.h
new file mode 100644
--- /dev/null
+++ b/Engine/Plugins/VirtualProduction/CaptureManager/CaptureManagerApp/Source/CaptureManagerMediaRW/Public/MediaRWManager.h
@@ -0,0 +1,61 @@
+// Copyright Epic Games, Inc. All Rights Reserved.
+
+#pragma once
+
+#include "IMediaRWFactory.h"
+
+#include "Containers/Map.h"
+
+class CAPTUREMANAGERMEDIARW_API FMediaRWManager
+{
+public:
+
+	FMediaRWManager();
+
+	FMediaRWManager(const FMediaRWManager& InCopy) = delete;
+	FMediaRWManager& operator=(const FMediaRWManager& InCopy) = delete;
+
+	void RegisterAudioReader(const TArray<FString>& InFormats, TUniquePtr<IAudioReaderFactory> InReader);
+	void RegisterVideoReader(const TArray<FString>& InFormats, TUniquePtr<IVideoReaderFactory> InReader);
+	void RegisterCalibrationReader(const TArray<FString>& InFormats, TUniquePtr<ICalibrationReaderFactory> InReader);
+
+	void RegisterAudioWriter(const TArray<FString>& InFormats, TUniquePtr<IAudioWriterFactory> InWriter);
+	void RegisterImageWriter(const TArray<FString>& InFormats, TUniquePtr<IImageWriterFactory> InWriter);
+	void RegisterCalibrationWriter(const TArray<FString>& InFormats, TUniquePtr<ICalibrationWriterFactory> InWriter);
+
+	TUniquePtr<IAudioReader> CreateAudioReaderByFormat(const FString& InFormat, int32 InIndex = 0);
+	TUniquePtr<IVideoReader> CreateVideoReaderByFormat(const FString& InFormat, int32 InIndex = 0);
+	TUniquePtr<ICalibrationReader> CreateCalibrationReaderByFormat(const FString& InFormat, int32 InIndex = 0);
+
+	TValueOrError<TUniquePtr<IAudioReader>, FText> CreateAudioReader(const FString& InPath, int32 InIndex = 0);
+	TValueOrError<TUniquePtr<IVideoReader>, FText> CreateVideoReader(const FString& InPath, int32 InIndex = 0);
+	TValueOrError<TUniquePtr<ICalibrationReader>, FText> CreateCalibrationReader(const FString& InPath, int32 InIndex = 0);
+
+	TUniquePtr<IAudioWriter> CreateAudioWriterByFormat(const FString& InFormat, int32 InIndex = 0);
+	TUniquePtr<IImageWriter> CreateImageWriterByFormat(const FString& InFormat, int32 InIndex = 0);
+	TUniquePtr<ICalibrationWriter> CreateCalibrationWriterByFormat(const FString& InFormat, int32 InIndex = 0);
+
+	TValueOrError<TUniquePtr<IAudioWriter>, FText> CreateAudioWriter(const FString& InDirectory, const FString& InFileName, const FString& InFormat, int32 InIndex = 0);
+	TValueOrError<TUniquePtr<IImageWriter>, FText> CreateImageWriter(const FString& InDirectory, const FString& InFileName, const FString& InFormat, int32 InIndex = 0);
+	TValueOrError<TUniquePtr<ICalibrationWriter>, FText> CreateCalibrationWriter(const FString& InDirectory, const FString& InFileName, const FString& InFormat, int32 InIndex = 0);
+
+private:
+
+	TMap<FString, IAudioReaderFactory*> AudioReadersPerFormat;
+	TArray<TUniquePtr<IAudioReaderFactory>> AudioReaders;
+
+	TMap<FString, IVideoReaderFactory*> VideoReadersPerFormat;
+	TArray<TUniquePtr<IVideoReaderFactory>> VideoReaders;
+
+	TMap<FString, ICalibrationReaderFactory*> CalibrationReadersPerFormat;
+	TArray<TUniquePtr<ICalibrationReaderFactory>> CalibrationReaders;
+
+	TMap<FString, IAudioWriterFactory*> AudioWritersPerFormat;
+	TArray<TUniquePtr<IAudioWriterFactory>> AudioWriters;
+
+	TMap<FString, IImageWriterFactory*> ImageWritersPerFormat;
+	TArray<TUniquePtr<IImageWriterFactory>> ImageWriters;
+
+	TMap<FString, ICalibrationWriterFactory*> CalibrationWritersPerFormat;
+	TArray<TUniquePtr<ICalibrationWriterFactory>> CalibrationWriters;
+};
\ No newline at end of file
diff --git a/Engine/Plugins/VirtualProduction/CaptureManager/CaptureManagerApp/Source/CaptureManagerMediaRW/Public/MediaSample.h b/Engine/Plugins/VirtualProduction/CaptureManager/CaptureManagerApp/Source/CaptureManagerMediaRW/Public/MediaSample.h
new file mode 100644
--- /dev/null
+++ b/Engine/Plugins/VirtualProduction/CaptureManager/CaptureManagerApp/Source/CaptureManagerMediaRW/Public/MediaSample.h
@@ -0,0 +1,163 @@
+// Copyright Epic Games, Inc. All Rights Reserved.
+
+#pragma once
+
+#include "IMediaAudioSample.h"
+#include "IMediaTextureSample.h"
+
+namespace UE::CaptureManager
+{
+
+enum class ESampleRate
+{
+	SR_44100Hz = 0,
+	SR_48000Hz,
+	SR_88200Hz,
+	SR_96000Hz,
+	SR_192000Hz
+};
+
+struct FMediaAudioSample
+{
+	TArray<uint8> Buffer;
+	uint32 Channels;
+	FTimespan Duration;
+	EMediaAudioSampleFormat SampleFormat;
+	uint32 Frames;
+	ESampleRate SampleRate;
+	FTimespan Time;
+};
+
+CAPTUREMANAGERMEDIARW_API int32 ConvertBitsPerSample(EMediaAudioSampleFormat InSampleFormat);
+CAPTUREMANAGERMEDIARW_API int32 ConvertSampleRate(ESampleRate InSampleRate);
+
+CAPTUREMANAGERMEDIARW_API ESampleRate ConvertSampleRate(int32 InSampleRate);
+
+enum class EMediaTexturePixelFormat
+{
+	Undefined = -1,
+	U8_RGB = 0,
+	U8_BGR,
+	U8_RGBA,
+	U8_BGRA,
+	U8_I444,
+	U8_I420,
+	U8_YUY2,
+	U8_NV12,
+	U8_Mono,
+	U16_Mono,
+	F_Mono,
+
+	Default = U8_RGB
+};
+
+struct FMediaTextureSample
+{
+	TArray<uint8> Buffer;
+	int32 Stride = -1;
+	FIntPoint Dimensions;
+	FTimespan Duration;
+	EMediaTexturePixelFormat CurrentFormat;
+	EMediaTexturePixelFormat DesiredFormat;
+	FTimespan Time;
+	EMediaOrientation Orientation;
+};
+
+CAPTUREMANAGERMEDIARW_API uint32 GetNumberOfChannels(EMediaTexturePixelFormat InPixelFormat);
+
+struct FOpenCVDistortionModel
+{
+	struct FRadial
+	{
+		double K1 = 0.0;
+		double K2 = 0.0;
+		double K3 = 0.0;
+	};
+
+	struct FTangential
+	{
+		double P1 = 0.0;
+		double P2 = 0.0;
+	};
+
+	FRadial Radial;
+	FTangential Tangential;
+};
+
+struct FIphoneDistortionModel
+{
+	TArray<double> LensDistortionTable;
+	TArray<double> InverseLensDistortionTable;
+};
+
+class CAPTUREMANAGERMEDIARW_API FCoordinateSystem
+{
+public:
+
+	enum FDirection
+	{
+		Front = 0,
+		Back,
+		Right,
+		Left,
+		Up,
+		Down
+	};
+
+	FCoordinateSystem();
+	FCoordinateSystem(FMatrix InMatDescription);
+	FCoordinateSystem(FDirection InXDirection, FDirection InYDirection, FDirection InZDirection);
+
+	const FMatrix& GetMatDescription() const;
+
+private:
+
+	static int32 GetIndex(FDirection InDirection);
+	static int32 GetSign(FDirection InDirection);
+
+	FMatrix MatDescription;
+};
+
+inline bool operator==(const FCoordinateSystem& InLeft, const FCoordinateSystem& InRight)
+{
+	return InLeft.GetMatDescription() == InRight.GetMatDescription();
+}
+
+extern const FCoordinateSystem UnrealCS;
+extern const FCoordinateSystem OpenCvCS;
+
+struct FMediaCalibrationSample
+{
+	enum ECameraType
+	{
+		Video = 0,
+		Depth,
+
+		Unknown
+	};
+
+	FString CameraId;
+	ECameraType CameraType = ECameraType::Unknown;
+
+	FVector2D FocalLength;
+	FVector2D PrincipalPoint;
+	FTransform Transform;
+
+	FIntPoint Dimensions;
+	EMediaOrientation Orientation;
+
+	TVariant<FEmptyVariantState, FIphoneDistortionModel, FOpenCVDistortionModel> DistortionModel;
+	
+	FCoordinateSystem InputCoordinateSystem;
+};
+
+CAPTUREMANAGERMEDIARW_API FVector ConvertToCoordinateSystem(const FVector& InVector, 
+															const FCoordinateSystem& InInputCoordinateSystem, 
+															const FCoordinateSystem& InOutputCoordinateSystem);
+
+CAPTUREMANAGERMEDIARW_API FTransform ConvertToCoordinateSystem(const FTransform& InTransform,
+															   const FCoordinateSystem& InInputCoordinateSystem,
+															   const FCoordinateSystem& InOutputCoordinateSystem);
+
+
+}
diff --git a/Engine/Plugins/VirtualProduction/CaptureManager/CaptureManagerApp/Source/CaptureManagerPipeline/Public/CaptureManagerPipeline.h b/Engine/Plugins/VirtualProduction/CaptureManager/CaptureManagerApp/Source/CaptureManagerPipeline/Public/CaptureManagerPipeline.h
new file mode 100644
--- /dev/null
+++ b/Engine/Plugins/VirtualProduction/CaptureManager/CaptureManagerApp/Source/CaptureManagerPipeline/Public/CaptureManagerPipeline.h
@@ -0,0 +1,52 @@
+// Copyright Epic Games, Inc. All Rights Reserved.
+
+#pragma once
+
+#include "HAL/Platform.h"
+
+#include "Nodes/ConvertAudioNode.h"
+#include "Nodes/ConvertDepthNode.h"
+#include "Nodes/ConvertVideoNode.h"
+#include "Nodes/ConvertCalibrationNode.h"
+
+#include "Templates/SharedPointer.h"
+#include "Misc/Guid.h"
+
+enum class EPipelineExecutionPolicy
+{
+	Asynchronous = 0,
+	Synchronous
+};
+
+class CAPTUREMANAGERPIPELINE_API FCaptureManagerPipeline : public TSharedFromThis<FCaptureManagerPipeline>
+{
+public:
+	using FResult = TMap<FGuid, FCaptureManagerPipelineNode::FResult>;
+
+	FCaptureManagerPipeline(EPipelineExecutionPolicy InExecutionPolicy);
+	~FCaptureManagerPipeline();
+
+	FGuid AddGenericNode(TSharedPtr<FCaptureManagerPipelineNode> InNode);
+	FGuid AddConvertVideoNode(TSharedPtr<FConvertVideoNode> InNode);
+	FGuid AddConvertAudioNode(TSharedPtr<FConvertAudioNode> InNode);
+	FGuid AddConvertDepthNode(TSharedPtr<FConvertDepthNode> InNode);
+	FGuid AddConvertCalibrationNode(TSharedPtr<FConvertCalibrationNode> InNode);
+
+	FGuid AddSyncedNode(TSharedPtr<FCaptureManagerPipelineNode> InNode);
+
+	// Blocking function
+	[[nodiscard]] FCaptureManagerPipeline::FResult Run();
+
+	void Cancel();
+
+private:
+
+	FGuid AddParallelPipelineNode(TSharedPtr<FCaptureManagerPipelineNode> InNode);
+
+	TPimplPtr<class FCaptureManagerPipelineImpl> Impl;
+
+	TMap<FGuid, TSharedPtr<FCaptureManagerPipelineNode>> ParallelNodes;
+	TMap<FGuid, TSharedPtr<FCaptureManagerPipelineNode>> SyncNodes;
+
+	EPipelineExecutionPolicy ExecutionPolicy;
+};
\ No newline at end of file
diff --git a/Engine/Plugins/VirtualProduction/CaptureManager/CaptureManagerApp/Source/CaptureManagerPipeline/Public/CaptureManagerPipelineNode.h b/Engine/Plugins/VirtualProduction/CaptureManager/CaptureManagerApp/Source/CaptureManagerPipeline/Public/CaptureManagerPipelineNode.h
new file mode 100644
--- /dev/null
+++ b/Engine/Plugins/VirtualProduction/CaptureManager/CaptureManagerApp/Source/CaptureManagerPipeline/Public/CaptureManagerPipelineNode.h
@@ -0,0 +1,50 @@
+// Copyright Epic Games, Inc. All Rights Reserved.
+
+#pragma once
+
+#include "HAL/Platform.h"
+#include "Internationalization/Text.h"
+#include "Templates/SharedPointer.h"
+#include "Templates/ValueOrError.h"
+
+#include "Templates/PimplPtr.h"
+
+class CAPTUREMANAGERPIPELINE_API FCaptureManagerPipelineError
+{
+public:
+	FCaptureManagerPipelineError(FText InMessage, int32 InCode = 0);
+
+	const FText& GetMessage() const;
+
+	int32 GetCode() const;
+
+private:
+
+	FText Message;
+	int32 Code;
+};
+
+class CAPTUREMANAGERPIPELINE_API FCaptureManagerPipelineNode : public TSharedFromThis<FCaptureManagerPipelineNode>
+{
+public:
+
+	using FResult = TValueOrError<void, FCaptureManagerPipelineError>;
+
+	FCaptureManagerPipelineNode(const FString& InName);
+	virtual ~FCaptureManagerPipelineNode();
+
+	FString GetName() const;
+
+	FResult Execute();
+	void Cancel();
+
+protected:
+
+	virtual FResult Prepare() = 0;
+	virtual FResult Run() = 0;
+	virtual FResult Validate() = 0;
+
+private:
+	class FCaptureManagerPipelineNodeImpl;
+	TPimplPtr<FCaptureManagerPipelineNodeImpl> Impl;
+};
\ No newline at end of file
diff --git a/Engine/Plugins/VirtualProduction/CaptureManager/CaptureManagerApp/Source/CaptureManagerPipeline/Public/Nodes/ConvertAudioNode.h b/Engine/Plugins/VirtualProduction/CaptureManager/CaptureManagerApp/Source/CaptureManagerPipeline/Public/Nodes/ConvertAudioNode.h
new file mode 100644
--- /dev/null
+++ b/Engine/Plugins/VirtualProduction/CaptureManager/CaptureManagerApp/Source/CaptureManagerPipeline/Public/Nodes/ConvertAudioNode.h
@@ -0,0 +1,29 @@
+// Copyright Epic Games, Inc. All Rights Reserved.
+
+#pragma once
+
+#include "CaptureManagerPipelineNode.h"
+
+#include "CaptureManagerTakeMetadata.h"
+
+class CAPTUREMANAGERPIPELINE_API FConvertAudioNode : public FCaptureManagerPipelineNode
+{
+public:
+
+	FConvertAudioNode(const FTakeMetadata::FAudio& InAudio,
+					  const FString& InOutputDirectory);
+
+protected:
+
+	FTakeMetadata::FAudio Audio;
+	FString OutputDirectory;
+
+private:
+
+	virtual FResult Prepare() override final;
+	virtual FResult Validate() override final;
+
+	FString GetAudioDirectory() const;
+
+	static FResult CheckForAudioFile(const FString& InAudioPath);
+};
\ No newline at end of file
diff --git a/Engine/Plugins/VirtualProduction/CaptureManager/CaptureManagerApp/Source/CaptureManagerPipeline/Public/Nodes/ConvertCalibrationNode.h b/Engine/Plugins/VirtualProduction/CaptureManager/CaptureManagerApp/Source/CaptureManagerPipeline/Public/Nodes/ConvertCalibrationNode.h
new file mode 100644
--- /dev/null
+++ b/Engine/Plugins/VirtualProduction/CaptureManager/CaptureManagerApp/Source/CaptureManagerPipeline/Public/Nodes/ConvertCalibrationNode.h
@@ -0,0 +1,26 @@
+// Copyright Epic Games, Inc. All Rights Reserved.
+
+#pragma once
+
+#include "CaptureManagerPipelineNode.h"
+
+#include "CaptureManagerTakeMetadata.h"
+
+class CAPTUREMANAGERPIPELINE_API FConvertCalibrationNode : public FCaptureManagerPipelineNode
+{
+public:
+
+	FConvertCalibrationNode(const FTakeMetadata::FCalibration& InCalibration,
+							const FString& InOutputDirectory);
+protected:
+
+	FTakeMetadata::FCalibration Calibration;
+	FString OutputDirectory;
+
+private:
+
+	virtual FResult Prepare() override final;
+	virtual FResult Validate() override final;
+
+	FString GetCalibrationDirectory() const;
+};
\ No newline at end of file
diff --git a/Engine/Plugins/VirtualProduction/CaptureManager/CaptureManagerApp/Source/CaptureManagerPipeline/Public/Nodes/ConvertDepthNode.h b/Engine/Plugins/VirtualProduction/CaptureManager/CaptureManagerApp/Source/CaptureManagerPipeline/Public/Nodes/ConvertDepthNode.h
new file mode 100644
--- /dev/null
+++ b/Engine/Plugins/VirtualProduction/CaptureManager/CaptureManagerApp/Source/CaptureManagerPipeline/Public/Nodes/ConvertDepthNode.h
@@ -0,0 +1,28 @@
+// Copyright Epic Games, Inc. All Rights Reserved.
+
+#pragma once
+
+#include "CaptureManagerPipelineNode.h"
+
+#include "CaptureManagerTakeMetadata.h"
+
+class CAPTUREMANAGERPIPELINE_API FConvertDepthNode : public FCaptureManagerPipelineNode
+{
+public:
+
+	FConvertDepthNode(const FTakeMetadata::FVideo& InDepth,
+					  const FString& InOutputDirectory);
+protected:
+
+	FTakeMetadata::FVideo Depth;
+	FString OutputDirectory;
+
+private:
+
+	virtual FResult Prepare() override final;
+	virtual FResult Validate() override final;
+
+	FString GetDepthDirectory() const;
+
+	static FResult CheckImagesForDepth(const FString& InDepthPath);
+};
\ No newline at end of file
diff --git a/Engine/Plugins/VirtualProduction/CaptureManager/CaptureManagerApp/Source/CaptureManagerPipeline/Public/Nodes/ConvertVideoNode.h b/Engine/Plugins/VirtualProduction/CaptureManager/CaptureManagerApp/Source/CaptureManagerPipeline/Public/Nodes/ConvertVideoNode.h
new file mode 100644
--- /dev/null
+++ b/Engine/Plugins/VirtualProduction/CaptureManager/CaptureManagerApp/Source/CaptureManagerPipeline/Public/Nodes/ConvertVideoNode.h
@@ -0,0 +1,29 @@
+// Copyright Epic Games, Inc. All Rights Reserved.
+
+#pragma once
+
+#include "CaptureManagerPipelineNode.h"
+
+#include "CaptureManagerTakeMetadata.h"
+
+class CAPTUREMANAGERPIPELINE_API FConvertVideoNode : public FCaptureManagerPipelineNode
+{
+public:
+
+	FConvertVideoNode(const FTakeMetadata::FVideo& InVideo, 
+					  const FString& InOutputDirectory);
+
+protected:
+
+	FTakeMetadata::FVideo Video;
+	FString OutputDirectory;
+
+private:
+
+	virtual FResult Prepare() override final;
+	virtual FResult Validate() override final;
+
+	FString GetVideoDirectory() const;
+
+	static FResult CheckImagesForVideo(const FString& InVideoPath);
+};
\ No newline at end of file
diff --git a/Engine/Plugins/VirtualProduction/CaptureManager/CaptureManagerApp/Source/CaptureManagerSettings/Private/Settings/CaptureManagerSettings.cpp b/Engine/Plugins/VirtualProduction/CaptureManager/CaptureManagerApp/Source/CaptureManagerSettings/Private/Settings/CaptureManagerSettings.cpp
new file mode 100644
--- /dev/null
+++ b/Engine/Plugins/VirtualProduction/CaptureManager/CaptureManagerApp/Source/CaptureManagerSettings/Private/Settings/CaptureManagerSettings.cpp
@@ -0,0 +1,37 @@
+// Copyright Epic Games, Inc. All Rights Reserved.
+
+#include "Settings/CaptureManagerSettings.h"
+#include "Misc/Paths.h"
+
+namespace UE::CaptureManager::Private
+{
+constexpr FStringView DefaultConversionDirectory = TEXTVIEW("CaptureManagerConversion");
+constexpr FStringView DefaultDownloadDirectory = TEXTVIEW("CaptureManagerDownload");
+}
+
+void UCaptureManagerSettings::Initialize()
+{
+	InitializeValuesIfNotSet();
+}
+
+void UCaptureManagerSettings::PostEditChangeProperty(struct FPropertyChangedEvent& InPropertyChangedEvent)
+{
+	Super::PostEditChangeProperty(InPropertyChangedEvent);
+
+	InitializeValuesIfNotSet();
+}
+
+void UCaptureManagerSettings::InitializeValuesIfNotSet()
+{
+	using namespace UE::CaptureManager;
+
+	if (DefaultConversionDirectory.Path.IsEmpty())
+	{
+		DefaultConversionDirectory.Path = FPaths::Combine(FPlatformProcess::UserTempDir(), Private::DefaultConversionDirectory);
+	}
+
+	if (DownloadDirectory.Path.IsEmpty())
+	{
+		DownloadDirectory.Path = FPaths::Combine(FPlatformProcess::UserTempDir(), Private::DefaultDownloadDirectory);
+	}
+}
\ No newline at end of file
diff --git a/Engine/Plugins/VirtualProduction/CaptureManager/CaptureManagerApp/Source/GenericCaptureSource/GenericCaptureSource.Build.cs b/Engine/Plugins/VirtualProduction/CaptureManager/CaptureManagerApp/Source/GenericCaptureSource/GenericCaptureSource.Build.cs
new file mode 100644
--- /dev/null
+++ b/Engine/Plugins/VirtualProduction/CaptureManager/CaptureManagerApp/Source/GenericCaptureSource/GenericCaptureSource.Build.cs
@@ -0,0 +1,41 @@
+// Copyright Epic Games, Inc. All Rights Reserved.
+
+using UnrealBuildTool;
+using System.IO;
+
+public class GenericCaptureSource : ModuleRules
+{
+	public GenericCaptureSource(ReadOnlyTargetRules Target) : base(Target)
+	{
+		PCHUsage = PCHUsageMode.UseExplicitOrSharedPCHs;
+		ShortName = "GenericCS";
+
+		PublicDependencyModuleNames.AddRange(new string[]
+		{
+			"CaptureUtils",
+			"Core"
+		});
+
+		PrivateDependencyModuleNames.AddRange(new string[]
+		{
+			"CaptureSourceFramework",
+			"CaptureProtocolStack",
+			"CaptureManagerPipeline",
+			"CaptureDataConverter",
+			"CaptureManagerMediaRW",
+			"CaptureCommunicationManager",
+			"DataIngestCore",
+			"Json",
+			"Media",
+			"LiveLinkFaceMetadata",
+			"LiveLinkHubCaptureMessaging",
+			"LiveLinkHubCaptureSourceWidgets",
+			"LiveLinkHubExportClient",
+			"SlateCore",
+			"StereoCameraMetadata",
+			"CaptureManagerTakeMetadata",
+			"CaptureManagerSettings",
+			"CoreUObject"
+		});
+	}
+}
