CaptureManager: Move plugin to Virtual Production directory

diff --git a/Engine/Plugins/VirtualProduction/CaptureManager/CaptureManagerApp/Source/GenericCaptureSource/Private/Archive/GenericArchiveCaptureSource.cpp b/Engine/Plugins/VirtualProduction/CaptureManager/CaptureManagerApp/Source/GenericCaptureSource/Private/Archive/GenericArchiveCaptureSource.cpp
new file mode 100644
--- /dev/null
+++ b/Engine/Plugins/VirtualProduction/CaptureManager/CaptureManagerApp/Source/GenericCaptureSource/Private/Archive/GenericArchiveCaptureSource.cpp
@@ -0,0 +1,514 @@
+// Copyright Epic Games, Inc. All Rights Reserved.
+
+#include "Archive/GenericArchiveCaptureSource.h"
+
+#include "Archive/GenericArchiveCaptureSourceFactory.h"
+#include "GenericCaptureSourceLog.h"
+
+#include "CaptureSourceFrameworkModule.h"
+#include "CaptureSourceDetailsCustomization.h"
+
+#include "HAL/PlatformFileManager.h"
+
+#include "CaptureCommunicationManagerModule.h"
+#include "LiveLinkFaceMetadata.h"
+#include "StereoCameraMetadataParseUtils.h"
+
+#include "Settings/CaptureManagerSettings.h"
+
+DEFINE_LOG_CATEGORY(LogGenericArchiveCaptureSource);
+
+#define LOCTEXT_NAMESPACE "GenericArchiveCaptureSource"
+
+namespace UE::CaptureManager::GenericArchiveCaptureSource
+{
+
+static TSharedRef<FCommunicationManager> GetCommunicationManager()
+{
+	FCaptureCommunicationManagerModule& CaptureCommunicationManagerModule =
+		FModuleManager::LoadModuleChecked<FCaptureCommunicationManagerModule>("CaptureCommunicationManager");
+	return CaptureCommunicationManagerModule.GetManager();
+}
+
+static EMediaTexturePixelFormat ConvertPixelFormat(EImagePixelFormat InImagePixelFormat)
+{
+	switch (InImagePixelFormat)
+	{
+		case EImagePixelFormat::U8_RGB:
+			return EMediaTexturePixelFormat::U8_RGB;
+		case EImagePixelFormat::U8_BGR:
+			return EMediaTexturePixelFormat::U8_BGR;
+		case EImagePixelFormat::U8_RGBA:
+			return EMediaTexturePixelFormat::U8_RGBA;
+		case EImagePixelFormat::U8_BGRA:
+			return EMediaTexturePixelFormat::U8_BGRA;
+		case EImagePixelFormat::U8_I444:
+			return EMediaTexturePixelFormat::U8_I444;
+		case EImagePixelFormat::U8_I420:
+			return EMediaTexturePixelFormat::U8_I420;
+		case EImagePixelFormat::U8_NV12:
+			return EMediaTexturePixelFormat::U8_NV12;
+		case EImagePixelFormat::U8_Mono:
+			return EMediaTexturePixelFormat::U8_Mono;
+		case EImagePixelFormat::U16_Mono:
+			return EMediaTexturePixelFormat::U16_Mono;
+		case EImagePixelFormat::F_Mono:
+			return EMediaTexturePixelFormat::F_Mono;
+		case EImagePixelFormat::Undefined:
+		default:
+			check(false);
+			return EMediaTexturePixelFormat::Undefined;
+	}
+}
+
+static FString GetOriginName(FTakeMetadataParserError::EOrigin InOrigin)
+{
+	switch (InOrigin)
+	{
+		case FTakeMetadataParserError::Reader:
+			return TEXT("Reader");
+		case FTakeMetadataParserError::Validator:
+			return TEXT("Validator");
+		case FTakeMetadataParserError::Parser:
+		default:
+			return TEXT("Parser");
+	}
+}
+
+}
+
+FGenericArchiveCaptureSource::FGenericArchiveCaptureSource(const FString& InName,
+														   FString InStoragePath)
+	: FCaptureSource(FGenericArchiveCaptureSourceFactory::Id, InName)
+	, StoragePath(MoveTemp(InStoragePath))
+	, CommunicationManager(UE::CaptureManager::GenericArchiveCaptureSource::GetCommunicationManager())
+{
+	using namespace UE::CaptureManager;
+
+	TUniquePtr<FIngestCapability> Ingest = MakeUnique<FIngestCapability>();
+	Ingest->SetImportTakeHandler(FCaptureSourceCapability::FCommandHandler::CreateRaw(this, &FGenericArchiveCaptureSource::ImportTakeHandler));
+	Ingest->SetUpdateTakeListHandler(FCaptureSourceCapability::FCommandHandler::CreateRaw(this, &FGenericArchiveCaptureSource::UpdateTakeListHandler));
+	Ingest->SetCancelHandler(FCaptureSourceCapability::FCommandHandler::CreateRaw(this, &FGenericArchiveCaptureSource::CancelImportHandler));
+
+	AddCapability(Ingest.Get());
+	OwnedCapabilities.Add(FIngestCapability::Name, MoveTemp(Ingest));
+}
+
+FGenericArchiveCaptureSource::~FGenericArchiveCaptureSource()
+{
+	Uploader = nullptr;
+
+	CommunicationManager->RemoveUploader(GetGuid());
+
+	Stop();
+}
+
+FCaptureVoidResult FGenericArchiveCaptureSource::Start()
+{
+	if (IFileManager::Get().DirectoryExists(*StoragePath))
+	{
+		PublishEvent<UE::CaptureManager::FReachableEvent>(true);
+	}
+
+	return MakeValue();
+}
+
+FCaptureVoidResult FGenericArchiveCaptureSource::Stop()
+{
+	PublishEvent<UE::CaptureManager::FReachableEvent>(false);
+
+	return MakeValue();
+}
+
+void FGenericArchiveCaptureSource::ImportTakeHandler(TSharedPtr<FCommandBase> InCommand)
+{
+	using namespace UE::CaptureManager;
+
+	TSharedPtr<FImportTakeCmd> ImportTakeCmd = StaticCastSharedPtr<FImportTakeCmd>(InCommand);
+
+	FCommandBase::FValues CommandValues = ImportTakeCmd->GetParamValues();
+	FTakeId TakeId = CommandValues[FIngestCapability::Id].Get<int64>();
+
+	FIngestCapability* Capability = GetCapability<FIngestCapability>(FIngestCapability::Name);
+	TOptional<FTakeMetadata> MaybeTakeMetadata = Capability->GetTakeMetadata(TakeId);
+	if (!MaybeTakeMetadata.IsSet())
+	{
+		FString Message = TEXT("Failed to find the take metadata for the specified take");
+		ImportTakeCmd->ProcessFinished(TakeId, MakeError(FIngestCapabilityError::InternalError, MoveTemp(Message)));
+		return;
+	}
+
+	FTakeMetadata TakeMetadata = MoveTemp(MaybeTakeMetadata.GetValue());
+
+	FString RelativeTakePath = RelativeTakePaths[TakeId];
+
+	TSharedPtr<FTaskProgress> TaskProgress 
+		= MakeShared<FTaskProgress>(2, FTaskProgress::FProgressReporter::CreateLambda([ImportTakeCmd](float InProgress)
+	{
+		FTakeId TakeId = ImportTakeCmd->GetParamValues()[FIngestCapability::Id].Get<int64>();
+		ImportTakeCmd->ProgressReporter(TakeId, InProgress);
+	}));
+
+	FString TakeName = TakeMetadata.Slate + TEXT("_") + FString::FromInt(TakeMetadata.TakeNumber);
+	FGuid TakeUploadId = FGuid::NewGuid();
+
+	const FString TakeConversionDirectory = ImportTakeCmd->IngestOptions.Location;
+
+	FCaptureDataConverterParams Params;
+	Params.TakeName = TakeName;
+	Params.TakeMetadata = TakeMetadata;
+	Params.TakeOriginDirectory = FPaths::Combine(StoragePath, RelativeTakePath);
+	Params.TakeOutputDirectory = FPaths::Combine(TakeConversionDirectory, TakeName + TEXT("_") + TakeUploadId.ToString());
+
+	FCaptureDataConverterResult<void> Result = RunConversion(TakeId, ImportTakeCmd->IngestOptions, TaskProgress, Params);
+
+	Converters.Remove(TakeId);
+
+	if (Result.HasValue())
+	{
+		if (!Uploader.IsValid())
+		{
+			Uploader = CommunicationManager->GetUploader(GetGuid(),
+														 GetName(),
+														 FUploader::FUploadStateUpdate::CreateRaw(this, &FGenericArchiveCaptureSource::UploadTakeStateHandler),
+														 FUploader::FUploadFinished::CreateRaw(this, &FGenericArchiveCaptureSource::UploadTakeFinished));
+		}
+
+		if (TSharedPtr<FUploader> SharedUploader = Uploader.Pin())
+		{
+			FScopeLock Lock(&Mutex);
+
+			FTaskProgress::FTask UploadTask = TaskProgress->StartTask();
+
+			ImportContexts.Add(TakeUploadId, { TakeName, ImportTakeCmd, TaskProgress, UploadTask });
+
+			int32 Worker = SharedUploader->UploadTake(TakeUploadId, Params.TakeOutputDirectory, MoveTemp(TakeMetadata));
+			WorkerForTake.Add(TakeId, Worker);
+		}
+		else
+		{
+			IFileManager::Get().DeleteDirectory(*Params.TakeOutputDirectory, false, true);
+
+			FString Message = TEXT("Failed to upload take: Uploader not found");
+			ImportTakeCmd->ProcessFinished(TakeId, MakeError(FIngestCapabilityError::InternalError, MoveTemp(Message)));
+		}
+	}
+	else
+	{
+		IFileManager::Get().DeleteDirectory(*Params.TakeOutputDirectory, false, true);
+
+		FCaptureDataConverterError Error = Result.StealError();
+
+		FText Message = FText::Join(FText::FromString(TEXT("\n\r")), Error.GetErrors());
+
+		ImportTakeCmd->ProcessFinished(TakeId, MakeError(FIngestCapabilityError::InternalError, Message.ToString()));
+	}
+}
+
+void FGenericArchiveCaptureSource::UpdateTakeListHandler(TSharedPtr<FCommandBase> InCommand)
+{
+	using namespace UE::CaptureManager;
+
+	TArray<FString> TakeMetadataFiles;
+	const bool bIterationResult = IFileManager::Get().IterateDirectoryRecursively(*StoragePath, [&](const TCHAR* InFileNameOrDirectory, bool bInIsDirectory)
+	{
+		if (!bInIsDirectory)
+		{
+			const FString CurrentFileName = FPaths::GetCleanFilename(InFileNameOrDirectory);
+			if (CurrentFileName == TEXT("take.json"))
+			{
+				TakeMetadataFiles.Add(InFileNameOrDirectory);
+			}
+			else if (FPaths::GetExtension(CurrentFileName) == FTakeMetadata::FileExtension)
+			{
+				TakeMetadataFiles.Add(InFileNameOrDirectory);
+			}
+		}
+
+		return true;
+	});
+
+	FIngestCapability* IngestCapability = GetCapability<FIngestCapability>(FIngestCapability::Name);
+
+	for (const FString& CurrentTakeFile : TakeMetadataFiles)
+	{
+		TOptional<FTakeMetadata> TakeInfoResult = ReadTake(CurrentTakeFile);
+
+		if (TakeInfoResult.IsSet())
+		{
+			FTakeId TakeId = IngestCapability->AddTakeMetadata(TakeInfoResult.GetValue());
+			FString CurrentDirectory = FPaths::GetPath(CurrentTakeFile);
+
+			FString TakePath = CurrentDirectory.Right(CurrentDirectory.Len() - StoragePath.Len());
+			FPaths::NormalizeDirectoryName(TakePath); // Removes trailing slash
+			RelativeTakePaths.Add(TakeId, MoveTemp(TakePath));
+		}
+	}
+
+	TSharedPtr<FUpdateTakeListCmd> Command = StaticCastSharedPtr<FUpdateTakeListCmd>(InCommand);
+	Command->Callback(IngestCapability->GetTakeIds());
+}
+
+void FGenericArchiveCaptureSource::CancelImportHandler(TSharedPtr<FCommandBase> InCommand)
+{
+	using namespace UE::CaptureManager;
+
+	FTakeId TakeId = InCommand->GetParamValue(FIngestCapability::Id).Get<int64>();
+
+	if (TUniquePtr<FCaptureDataConverter>* Converter = Converters.Find(TakeId))
+	{
+		(*Converter)->Cancel();
+	}
+
+	int32 WorkerId = INDEX_NONE;
+
+	{
+		FScopeLock Lock(&Mutex);
+		if (int32* WorkerIdElem = WorkerForTake.Find(TakeId))
+		{
+			WorkerId = *WorkerIdElem;
+		}
+	}
+
+	if (WorkerId != INDEX_NONE)
+	{
+		if (TSharedPtr<FUploader> SharedUploader = Uploader.Pin())
+		{
+			SharedUploader->AbortUpload(WorkerId);
+		}
+		else
+		{
+			// If `UploadTakeFinished` is not called (eg. connection between LLH and UE broke), `Uploader` will be deleted, but UI will stay in state where import is in progress 99.9%.
+			// Clicking stop button at that moment will not enter if block above, as uploader does not exist, so in that case we have to manually cleanup the workspace and put UI in proper state.
+			FScopeLock Lock(&Mutex);
+
+			TMap<FGuid, FContext> FilteredMap = ImportContexts.FilterByPredicate([TakeId](const TPair<FGuid, FContext>& Item) -> bool
+			{
+				return Item.Value.ImportCmd->GetParamValue(FIngestCapability::Id).Get<int64>() == TakeId;
+			});
+
+			check(FilteredMap.Num() == 1);
+			TMap<FGuid, FContext>::TConstIterator FoundContext = FilteredMap.CreateConstIterator();
+			check(FoundContext);
+
+			const FContext& Context = FoundContext->Value;
+
+			const FString TakeDirectory = FPaths::Combine(Context.ImportCmd->IngestOptions.Location, 
+														  Context.TakeName + TEXT("_") + FoundContext->Key.ToString());
+			IFileManager::Get().DeleteDirectory(*TakeDirectory, false, true);
+
+			Context.ImportCmd->ProcessFinished(TakeId, MakeError(FIngestCapabilityError::InternalError, "Aborted by user"));
+
+			ImportContexts.Remove(FoundContext->Key);
+
+			WorkerForTake.Remove(TakeId);
+		}
+	}
+}
+
+FCaptureDataConverterResult<void> FGenericArchiveCaptureSource::RunConversion(UE::CaptureManager::FTakeId InTakeId,
+																			  const UE::CaptureManager::FIngestOptions& InIngestOptions,
+																			  TSharedPtr<UE::CaptureManager::FTaskProgress> InTaskProgress,
+																			  FCaptureDataConverterParams& InParams)
+{
+	using namespace UE::CaptureManager;
+
+	FTaskProgress::FTask ConvertTask = InTaskProgress->StartTask();
+
+	FCaptureConvertVideoOutputParams VideoParams;
+	VideoParams.Format = InIngestOptions.Video.Format;
+	VideoParams.ImageFileName = InIngestOptions.Video.FileNamePrefix;
+	VideoParams.OutputPixelFormat = GenericArchiveCaptureSource::ConvertPixelFormat(InIngestOptions.Video.PixelFormat);
+
+	InParams.VideoOutputParams = MoveTemp(VideoParams);
+
+	FCaptureConvertAudioOutputParams AudioParams;
+	AudioParams.Format = InIngestOptions.Audio.Format;
+	AudioParams.AudioFileName = InIngestOptions.Audio.FileNamePrefix;
+
+	InParams.AudioOutputParams = MoveTemp(AudioParams);
+
+	FCaptureConvertDepthOutputParams DepthParams;
+	DepthParams.ImageFileName = TEXT("depth");
+
+	InParams.DepthOutputParams = MoveTemp(DepthParams);
+
+	FCaptureConvertCalibrationOutputParams CalibrationParams;
+	CalibrationParams.FileName = TEXT("calibration");
+
+	InParams.CalibrationOutputParams = MoveTemp(CalibrationParams);
+
+	FCaptureDataConverter::FProgressReporter ProgressReporter =
+		FCaptureDataConverter::FProgressReporter::CreateLambda([ConvertTask](const float InProgress) mutable
+	{
+		ConvertTask.Update(InProgress);
+	});
+
+	TUniquePtr<FCaptureDataConverter>& Converter = Converters.FindOrAdd(InTakeId, MakeUnique<FCaptureDataConverter>());
+	return Converter->Run(InParams, MoveTemp(ProgressReporter));
+}
+
+void FGenericArchiveCaptureSource::UploadTakeStateHandler(const FGuid& InTakeUploadId, double InProgress)
+{
+	if (!ensure(InTakeUploadId.IsValid()))
+	{
+		return;
+	}
+
+	FScopeLock Lock(&Mutex);
+	if (FContext* Context = ImportContexts.Find(InTakeUploadId))
+	{
+		Context->Task.Update(InProgress);
+	}
+}
+
+void FGenericArchiveCaptureSource::UploadTakeFinished(const FGuid& InTakeUploadId, const FString& InMessage, int32 InCode)
+{
+	using namespace UE::CaptureManager;
+
+	if (!ensure(InTakeUploadId.IsValid()))
+	{
+		return;
+	}
+
+	FScopeLock Lock(&Mutex);
+	if (const FContext* Context = ImportContexts.Find(InTakeUploadId))
+	{
+		const TSharedPtr<FImportTakeCmd>& ImportTakeCmd = Context->ImportCmd;
+		check(ImportTakeCmd);
+
+		if (!ImportTakeCmd)
+		{
+			return;
+		}
+
+		const bool bShouldCleanConversionDirectory = GetDefault<UCaptureManagerSettings>()->bShouldCleanConversionDirectory;
+		if (bShouldCleanConversionDirectory)
+		{
+			const FString ConversionDirectory = FPaths::Combine(ImportTakeCmd->IngestOptions.Location, 
+																Context->TakeName + TEXT("_") + InTakeUploadId.ToString());
+			IFileManager::Get().DeleteDirectory(*ConversionDirectory, false, true);
+		}
+
+		FTakeId TakeId = ImportTakeCmd->GetParamValue(FIngestCapability::Id).Get<int64>();
+
+		if (InMessage.IsEmpty())
+		{
+			ImportTakeCmd->ProcessFinished(TakeId, MakeValue());
+		}
+		else
+		{
+			ImportTakeCmd->ProcessFinished(TakeId, MakeError(FIngestCapabilityError::InternalError, InMessage));
+		}
+
+		ImportContexts.Remove(InTakeUploadId);
+
+		WorkerForTake.Remove(TakeId);
+	}
+}
+
+TOptional<FTakeMetadata> FGenericArchiveCaptureSource::ReadTake(const FString& InCurrentTakeFile)
+{
+	using namespace UE::CaptureManager;
+
+	TValueOrError<FTakeMetadata, FTakeMetadataParserError> TakeMetadataResult = TakeMetadataParser.Parse(InCurrentTakeFile);
+	FString CurrentDirectory = FPaths::GetPath(InCurrentTakeFile);
+
+	if (TakeMetadataResult.HasError())
+	{
+		FTakeMetadataParserError ParserError = TakeMetadataResult.StealError();
+		FString OriginName = GenericArchiveCaptureSource::GetOriginName(ParserError.Origin);
+
+		FString Message =
+			FString::Format(TEXT("Unable to parse take metadata file - {0} (Error origin: {1}): {2}"),
+							{
+								InCurrentTakeFile,
+								ParserError.Message.ToString(),
+								OriginName
+							});
+
+		UE_LOG(LogGenericArchiveCaptureSource, Warning, TEXT("%s"), *Message);
+
+		UE_LOG(LogGenericArchiveCaptureSource, Display, TEXT("Checking backwards compatible take metadata formats"), *CurrentDirectory);
+
+		TArray<FText> ValidationFailures;
+
+		UE_LOG(LogGenericArchiveCaptureSource, Display, TEXT("Checking directory (%s) for pre UE 5.6 LiveLink take metadata format"), *CurrentDirectory);
+		TOptional<FTakeMetadata> ParseOldMetadataResult = LiveLinkMetadata::ParseOldLiveLinkTakeMetadata(CurrentDirectory, ValidationFailures);
+
+		if (ParseOldMetadataResult.IsSet())
+		{
+			return ParseOldMetadataResult.GetValue();
+		}
+
+		UE_LOG(LogGenericArchiveCaptureSource, Display, TEXT("Checking directory (%s) for pre UE 5.6 StereoCamera take metadata format"), *CurrentDirectory);
+		ParseOldMetadataResult = StereoCameraMetadata::ParseOldStereoCameraMetadata(CurrentDirectory, ValidationFailures);
+
+		if (ParseOldMetadataResult.IsSet())
+		{
+			return ParseOldMetadataResult.GetValue();
+		}
+
+		UE_LOG(LogGenericArchiveCaptureSource, Error, TEXT("Failed to parse take metadata file: %s"), *CurrentDirectory);
+		return {};
+	}
+
+	FTakeMetadata& TakeMetadata = TakeMetadataResult.GetValue();
+
+	TOptional<FString> ThumbnailPath = TakeMetadata.Thumbnail.GetThumbnailPath();
+	if (ThumbnailPath.IsSet())
+	{
+		if (FPaths::IsRelative(*ThumbnailPath))
+		{
+			TakeMetadata.Thumbnail = FTakeThumbnailData(FPaths::ConvertRelativePathToFull(CurrentDirectory, *ThumbnailPath));
+		}
+	}
+
+	if (!TakeMetadata.Video.IsEmpty())
+	{
+		for (FTakeMetadata::FVideo& Video : TakeMetadata.Video)
+		{
+			if (FPaths::IsRelative(*Video.Path))
+			{
+				Video.Path = FPaths::ConvertRelativePathToFull(CurrentDirectory, *Video.Path);
+			}
+		}
+	}
+
+	if (!TakeMetadata.Depth.IsEmpty())
+	{
+		for (FTakeMetadata::FVideo& Depth : TakeMetadata.Depth)
+		{
+			if (FPaths::IsRelative(*Depth.Path))
+			{
+				Depth.Path = FPaths::ConvertRelativePathToFull(CurrentDirectory, *Depth.Path);
+			}
+		}
+	}
+
+	if (!TakeMetadata.Audio.IsEmpty())
+	{
+		for (FTakeMetadata::FAudio& Audio : TakeMetadata.Audio)
+		{
+			if (FPaths::IsRelative(*Audio.Path))
+			{
+				Audio.Path = FPaths::ConvertRelativePathToFull(CurrentDirectory, *Audio.Path);
+			}
+		}
+	}
+
+	if (!TakeMetadata.Calibration.IsEmpty())
+	{
+		for (FTakeMetadata::FCalibration& Calib : TakeMetadata.Calibration)
+		{
+			if (FPaths::IsRelative(*Calib.Path))
+			{
+				Calib.Path = FPaths::ConvertRelativePathToFull(CurrentDirectory, *Calib.Path);
+			}
+		}
+	}
+	
+	return TakeMetadataResult.StealValue();
+}
+
+#undef LOCTEXT_NAMESPACE
diff --git a/Engine/Plugins/VirtualProduction/CaptureManager/CaptureManagerApp/Source/GenericCaptureSource/Private/Archive/GenericArchiveCaptureSource.h b/Engine/Plugins/VirtualProduction/CaptureManager/CaptureManagerApp/Source/GenericCaptureSource/Private/Archive/GenericArchiveCaptureSource.h
new file mode 100644
--- /dev/null
+++ b/Engine/Plugins/VirtualProduction/CaptureManager/CaptureManagerApp/Source/GenericCaptureSource/Private/Archive/GenericArchiveCaptureSource.h
@@ -0,0 +1,71 @@
+// Copyright Epic Games, Inc. All Rights Reserved.
+
+#pragma once
+
+#include "CaptureSource.h"
+
+#include "Capabilities/IngestCapability.h"
+
+#include "Async/StopToken.h"
+#include "Async/TaskProgress.h"
+
+#include "CaptureManagerTakeMetadata.h"
+
+#include "LiveLinkHubExportManager.h"
+
+#include "CaptureDataConverter.h"
+
+#include "CaptureCommunicationManager.h"
+
+class FGenericArchiveCaptureSource final : public UE::CaptureManager::FCaptureSource
+{
+public:
+
+	FGenericArchiveCaptureSource(const FString& InName,
+								 FString StoragePath);
+
+	~FGenericArchiveCaptureSource();
+
+	virtual FCaptureVoidResult Start() override;
+	virtual FCaptureVoidResult Stop() override;
+
+private:
+
+	void ImportTakeHandler(TSharedPtr<FCommandBase> InCommand);
+	void UpdateTakeListHandler(TSharedPtr<FCommandBase> InCommand);
+	void CancelImportHandler(TSharedPtr<FCommandBase> InCommand);
+
+	FCaptureDataConverterResult<void> RunConversion(UE::CaptureManager::FTakeId InTakeId,
+													const UE::CaptureManager::FIngestOptions& InIngestOptions,
+													TSharedPtr<UE::CaptureManager::FTaskProgress> InTaskProgress,
+													FCaptureDataConverterParams& InParams);
+
+	void UploadTakeStateHandler(const FGuid& InTakeUploadId, double InProgress);
+	void UploadTakeFinished(const FGuid& InTakeUploadId, const FString& InMessage, int32 InCode);
+
+	TOptional<FTakeMetadata> ReadTake(const FString& InCurrentTakeFile);
+
+	FString StoragePath;
+
+	TMap<FString, TUniquePtr<UE::CaptureManager::FCaptureSourceCapability>> OwnedCapabilities;
+
+	struct FContext
+	{
+		FString TakeName;
+		TSharedPtr<UE::CaptureManager::FImportTakeCmd> ImportCmd;
+		TSharedPtr<UE::CaptureManager::FTaskProgress> TaskProgress;
+		UE::CaptureManager::FTaskProgress::FTask Task;
+	};
+
+	TSharedRef<FCommunicationManager> CommunicationManager;
+	TWeakPtr<FUploader> Uploader;
+
+	TMap<FGuid, FContext> ImportContexts;
+	TMap<UE::CaptureManager::FTakeId, FString> RelativeTakePaths;
+
+	FCriticalSection Mutex;
+	TMap<UE::CaptureManager::FTakeId, int32> WorkerForTake;
+
+	TMap<UE::CaptureManager::FTakeId, TUniquePtr<FCaptureDataConverter>> Converters;
+	FTakeMetadataParser TakeMetadataParser;
+};
\ No newline at end of file
diff --git a/Engine/Plugins/VirtualProduction/CaptureManager/CaptureManagerApp/Source/GenericCaptureSource/Private/Archive/GenericArchiveCaptureSourceFactory.cpp b/Engine/Plugins/VirtualProduction/CaptureManager/CaptureManagerApp/Source/GenericCaptureSource/Private/Archive/GenericArchiveCaptureSourceFactory.cpp
new file mode 100644
--- /dev/null
+++ b/Engine/Plugins/VirtualProduction/CaptureManager/CaptureManagerApp/Source/GenericCaptureSource/Private/Archive/GenericArchiveCaptureSourceFactory.cpp
@@ -0,0 +1,60 @@
+// Copyright Epic Games, Inc. All Rights Reserved.
+
+#include "GenericArchiveCaptureSourceFactory.h"
+
+#include "GenericArchiveCaptureSource.h"
+
+const FString FGenericArchiveCaptureSourceFactory::Id = TEXT("GenericArchiveCaptureSource");
+const FText FGenericArchiveCaptureSourceFactory::Name = NSLOCTEXT("GenericArchiveCaptureSourceFactory", "FactoryName", "Generic Archive");
+
+const FString FGenericArchiveCaptureSourceFactory::StoragePath = TEXT("StoragePath");
+
+FGenericArchiveCaptureSourceFactory::FGenericArchiveCaptureSourceFactory() = default;
+
+FText FGenericArchiveCaptureSourceFactory::GetCaptureSourceFactoryName() const
+{
+	return Name;
+}
+
+FString FGenericArchiveCaptureSourceFactory::GetCaptureSourceFactoryId() const
+{
+	return Id;
+}
+
+bool FGenericArchiveCaptureSourceFactory::IsDiscoverable() const
+{
+	return false;
+}
+
+void FGenericArchiveCaptureSourceFactory::DiscoverImpl(FDiscoverCallback InCallback, FCaptureVoidResult& OutResult)
+{
+}
+
+FGenericArchiveCaptureSourceFactory::FCaptureSourceResult
+FGenericArchiveCaptureSourceFactory::CreateCaptureSourceImpl(const FString& InName, const TMap<FString, FPropertyValue>& InCreationParamsValue)
+{
+	if (InCreationParamsValue.Num() != 1)
+	{
+		return MakeError(FCaptureSourceError(ECaptureSourceError::MissingArgument));
+	}
+
+	if (!InCreationParamsValue.Contains(StoragePath))
+	{
+		return MakeError(FCaptureSourceError(ECaptureSourceError::MissingArgument));
+	}
+
+	FString StoragePathValue = InCreationParamsValue[StoragePath].Get<FString>();
+
+	TUniquePtr<FGenericArchiveCaptureSource> Source = MakeUnique<FGenericArchiveCaptureSource>(InName, StoragePathValue);
+	
+	Source->SetCreationParams(InCreationParamsValue);
+
+	return MakeValue(MoveTemp(Source));
+}
+
+void FGenericArchiveCaptureSourceFactory::CreateCaptureSourceDescriptorImpl(FCaptureSourceDescriptor& OutDescriptor) const
+{
+	OutDescriptor.SetDiscoverable(IsDiscoverable());
+
+	OutDescriptor.AddCreationParam(FPropertyDesc(StoragePath, FPropertyDesc::EType::String));
+}
\ No newline at end of file
diff --git a/Engine/Plugins/VirtualProduction/CaptureManager/CaptureManagerApp/Source/GenericCaptureSource/Private/Archive/GenericArchiveCaptureSourceFactory.h b/Engine/Plugins/VirtualProduction/CaptureManager/CaptureManagerApp/Source/GenericCaptureSource/Private/Archive/GenericArchiveCaptureSourceFactory.h
new file mode 100644
--- /dev/null
+++ b/Engine/Plugins/VirtualProduction/CaptureManager/CaptureManagerApp/Source/GenericCaptureSource/Private/Archive/GenericArchiveCaptureSourceFactory.h
@@ -0,0 +1,29 @@
+// Copyright Epic Games, Inc. All Rights Reserved.
+
+#pragma once
+
+#include "CaptureSourceFactory.h"
+
+class FGenericArchiveCaptureSourceFactory final : public FCaptureSourceFactory
+{
+public:
+
+	static const FString Id;
+	static const FText Name;
+
+	static const FString StoragePath;
+
+	FGenericArchiveCaptureSourceFactory();
+
+	virtual FText GetCaptureSourceFactoryName() const override;
+	virtual FString GetCaptureSourceFactoryId() const override;
+
+	virtual bool IsDiscoverable() const;
+
+private:
+	virtual void DiscoverImpl(FDiscoverCallback InCallback, FCaptureVoidResult& OutResult) override;
+
+	virtual FCaptureSourceResult CreateCaptureSourceImpl(const FString& InName, const TMap<FString, FPropertyValue>& InCreationParamsValue) override;
+	virtual void CreateCaptureSourceDescriptorImpl(FCaptureSourceDescriptor& OutDescriptor) const override;
+
+};
\ No newline at end of file
diff --git a/Engine/Plugins/VirtualProduction/CaptureManager/CaptureManagerApp/Source/GenericCaptureSource/Private/GenericCaptureSourceLog.h b/Engine/Plugins/VirtualProduction/CaptureManager/CaptureManagerApp/Source/GenericCaptureSource/Private/GenericCaptureSourceLog.h
new file mode 100644
--- /dev/null
+++ b/Engine/Plugins/VirtualProduction/CaptureManager/CaptureManagerApp/Source/GenericCaptureSource/Private/GenericCaptureSourceLog.h
@@ -0,0 +1,9 @@
+// Copyright Epic Games, Inc. All Rights Reserved.
+
+#pragma once
+
+#include "Logging/LogMacros.h"
+
+DECLARE_LOG_CATEGORY_EXTERN(LogGenericArchiveCaptureSource, Log, All);
+DECLARE_LOG_CATEGORY_EXTERN(LogGenericCPSCaptureSource, Log, All);
+DECLARE_LOG_CATEGORY_EXTERN(LogCPSDevice, Log, All);
diff --git a/Engine/Plugins/VirtualProduction/CaptureManager/CaptureManagerApp/Source/GenericCaptureSource/Private/GenericCaptureSourceModule.cpp b/Engine/Plugins/VirtualProduction/CaptureManager/CaptureManagerApp/Source/GenericCaptureSource/Private/GenericCaptureSourceModule.cpp
new file mode 100644
--- /dev/null
+++ b/Engine/Plugins/VirtualProduction/CaptureManager/CaptureManagerApp/Source/GenericCaptureSource/Private/GenericCaptureSourceModule.cpp
@@ -0,0 +1,56 @@
+// Copyright Epic Games, Inc. All Rights Reserved.
+
+#include "GenericCaptureSourceModule.h"
+
+#include "CaptureSourceFrameworkModule.h"
+
+#include "Archive/GenericArchiveCaptureSourceFactory.h"
+#include "Live/GenericCPSCaptureSourceFactory.h"
+
+#include "Widgets/SArchiveCreationPanel.h"
+#include "Widgets/SConnectionCreationPanel.h"
+
+template<typename FactoryType, typename WidgetType>
+static void RegisterFactory(UE::CaptureManager::FCaptureSourceFrameworkModule& CaptureSourceFrameworkModule)
+{
+	using namespace UE::CaptureManager;
+
+	TUniquePtr<FactoryType> Factory = MakeUnique<FactoryType>();
+	const FString FactoryId = Factory->GetCaptureSourceFactoryId();
+
+	// Override the built-in (dynamically generated) user interface, so that we have a little more control over the creation panels
+	CaptureSourceFrameworkModule.GetCustomization().RegisterCaptureSourceCustomization(
+		FactoryId,
+		FCaptureSourceDetailsCustomization::FCaptureSourceWidgetCreator::CreateLambda(
+			[FactoryId]
+			(
+				[[maybe_unused]] const FCaptureSourceDescriptor& InCaptureSourceDescriptor,
+				[[maybe_unused]] FOnCaptureSourceConfigured InCaptureSourceConfigured
+			)
+			{
+				TSharedPtr<SWidget> Widget = SNew(WidgetType)
+					.FactoryId(FactoryId)
+					.OnSourceConfigured(MoveTemp(InCaptureSourceConfigured));
+				return Widget;
+			}
+		)
+	);
+
+	CaptureSourceFrameworkModule.GetCaptureSourceManager().RegisterCaptureSourceFactory(MoveTemp(Factory));
+}
+
+void FGenericCaptureSourceModule::StartupModule()
+{
+	using namespace UE::CaptureManager;
+
+	FCaptureSourceFrameworkModule& CaptureSourceFrameworkModule = FModuleManager::LoadModuleChecked<FCaptureSourceFrameworkModule>("CaptureSourceFramework");
+
+	RegisterFactory<FGenericArchiveCaptureSourceFactory, SArchiveCreationPanel>(CaptureSourceFrameworkModule);
+	RegisterFactory<FGenericCPSCaptureSourceFactory, SConnectionCreationPanel>(CaptureSourceFrameworkModule);
+}
+
+void FGenericCaptureSourceModule::ShutdownModule()
+{
+}
+
+IMPLEMENT_MODULE(FGenericCaptureSourceModule, GenericCaptureSource)
\ No newline at end of file
diff --git a/Engine/Plugins/VirtualProduction/CaptureManager/CaptureManagerApp/Source/GenericCaptureSource/Private/GenericCaptureSourceModule.h b/Engine/Plugins/VirtualProduction/CaptureManager/CaptureManagerApp/Source/GenericCaptureSource/Private/GenericCaptureSourceModule.h
new file mode 100644
--- /dev/null
+++ b/Engine/Plugins/VirtualProduction/CaptureManager/CaptureManagerApp/Source/GenericCaptureSource/Private/GenericCaptureSourceModule.h
@@ -0,0 +1,16 @@
+// Copyright Epic Games, Inc. All Rights Reserved.
+
+#pragma once
+
+#include "Modules/ModuleManager.h"
+
+class FGenericCaptureSourceModule
+	: public IModuleInterface
+{
+public:
+
+	//~Begin IModuleInterface interface
+	virtual void StartupModule() override;
+	virtual void ShutdownModule() override;
+	//~End IModuleInterface interface
+};
\ No newline at end of file
diff --git a/Engine/Plugins/VirtualProduction/CaptureManager/CaptureManagerApp/Source/GenericCaptureSource/Private/Live/CPS/CPSDevice.h b/Engine/Plugins/VirtualProduction/CaptureManager/CaptureManagerApp/Source/GenericCaptureSource/Private/Live/CPS/CPSDevice.h
new file mode 100644
--- /dev/null
+++ b/Engine/Plugins/VirtualProduction/CaptureManager/CaptureManagerApp/Source/GenericCaptureSource/Private/Live/CPS/CPSDevice.h
@@ -0,0 +1,135 @@
+// Copyright Epic Games, Inc. All Rights Reserved.
+
+#pragma once
+
+#include "Control/ControlMessenger.h"
+
+#include "ExportClient/ExportClient.h"
+
+#include "Async/EventSourceUtils.h"
+#include "Misc/Optional.h"
+
+#include "Capabilities/IngestCapability.h"
+
+#include "Async/CaptureTimerManager.h"
+
+struct FConnectionStateChangedEvent : public UE::CaptureManager::FCaptureEvent
+{
+	inline static const FString Name = TEXT("ConnectionStateChanged");
+
+	enum class EState
+	{
+		Unknown = 0,
+		Connected,
+		Disconnected
+	};
+
+	FConnectionStateChangedEvent(EState InConnectionState);
+
+	EState ConnectionState;
+};
+
+struct FCPSEvent : public UE::CaptureManager::FCaptureEvent
+{
+	inline static const FString Name = TEXT("CPSEvent");
+
+	FCPSEvent(TSharedPtr<UE::CaptureManager::FControlUpdate> InUpdateMessage);
+
+	TSharedPtr<UE::CaptureManager::FControlUpdate> UpdateMessage;
+};
+
+struct FCPSStateEvent : public UE::CaptureManager::FCaptureEvent
+{
+	inline static const FString Name = TEXT("CPSStateEvent");
+
+	FCPSStateEvent(UE::CaptureManager::FGetStateResponse InGetStateResponse);
+
+	UE::CaptureManager::FGetStateResponse GetStateResponse;
+};
+
+class FCPSDevice : 
+	public UE::CaptureManager::FCaptureEventSource, 
+	public TSharedFromThis<FCPSDevice>
+{
+private:
+
+	struct FPrivateToken
+	{
+		explicit FPrivateToken() = default;
+	};
+
+public:
+
+	static TSharedPtr<FCPSDevice> MakeCPSDevice(FString InDeviceIpAddress, uint16 InDevicePort);
+
+	FCPSDevice(FPrivateToken,
+			   FString InDeviceIpAddress,
+			   uint16 InDevicePort);
+
+	~FCPSDevice();
+
+	void InitiateConnect();
+	void Stop();
+
+	bool IsConnected();
+
+	UE::CaptureManager::TProtocolResult<void> StartRecording(FString SlateName,
+															 uint16 TakeNumber,
+															 TOptional<FString> Subject = TOptional<FString>(),
+															 TOptional<FString> Scenario = TOptional<FString>(),
+															 TOptional<TArray<FString>> Tags = TOptional<TArray<FString>>());
+
+	UE::CaptureManager::TProtocolResult<void> StopRecording();
+
+	UE::CaptureManager::TProtocolResult<TArray<UE::CaptureManager::FGetTakeMetadataResponse::FTakeObject>> FetchTakeList();
+	UE::CaptureManager::TProtocolResult<UE::CaptureManager::FGetTakeMetadataResponse::FTakeObject> FetchTake(const FString& InTakeName);
+
+	void AddTakeMetadata(UE::CaptureManager::FTakeId InId, UE::CaptureManager::FGetTakeMetadataResponse::FTakeObject InTake);
+	void RemoveTakeMetadata(UE::CaptureManager::FTakeId InId);
+	UE::CaptureManager::FGetTakeMetadataResponse::FTakeObject GetTake(UE::CaptureManager::FTakeId InId);
+	UE::CaptureManager::FTakeId GetTakeId(const FString& InTakeName);
+
+	void StartExport(UE::CaptureManager::FTakeId InTakeId, TUniquePtr<UE::CaptureManager::FBaseStream> InStream);
+	void CancelExport(UE::CaptureManager::FTakeId InTakeId);
+
+	void FetchThumbnails(TUniquePtr<UE::CaptureManager::FBaseStream> InStream);
+
+private:
+
+	inline static const float ConnectInterval = 5.0f;
+
+	void InitializeDelegates();
+
+	struct FEmpty {};
+
+	void ConnectControlClient(FEmpty);
+
+	void RegisterForAllEvents();
+	void OnCPSEvent(TSharedPtr<UE::CaptureManager::FControlUpdate> InUpdateMessage);
+
+	void StartConnectTimer(float InDelay = 0.0f);
+	void OnConnectTick();
+	void OnDisconnect(const FString& InCause);
+
+	static TSharedRef<UE::CaptureManager::FCaptureTimerManager> GetTimerManager();
+
+	TSharedRef<UE::CaptureManager::FCaptureTimerManager> TimerManager;
+
+	FString DeviceIpAddress;
+	uint16 DeviceControlPort;
+
+	std::atomic_bool bIsConnected;
+
+	using FConnectionThread = UE::CaptureManager::TQueueRunner<FEmpty>;
+	TUniquePtr<FConnectionThread> ConnThread;
+	UE::CaptureManager::FControlMessenger ControlMessenger;
+	UE::CaptureManager::FCaptureTimerManager::FTimerHandle ConnectTimerHandle;
+
+	TUniquePtr<UE::CaptureManager::FExportClient> ExportClient;
+
+	FCriticalSection Mutex;
+	TMap<UE::CaptureManager::FTakeId, UE::CaptureManager::FGetTakeMetadataResponse::FTakeObject> TakeMetadata;
+
+	TMap<UE::CaptureManager::FTakeId, UE::CaptureManager::FExportClient::FTaskId> IdMap;
+};
+
diff --git a/Engine/Plugins/VirtualProduction/CaptureManager/CaptureManagerApp/Source/GenericCaptureSource/Private/Live/GenericCPSCaptureSource.cpp b/Engine/Plugins/VirtualProduction/CaptureManager/CaptureManagerApp/Source/GenericCaptureSource/Private/Live/GenericCPSCaptureSource.cpp
new file mode 100644
--- /dev/null
+++ b/Engine/Plugins/VirtualProduction/CaptureManager/CaptureManagerApp/Source/GenericCaptureSource/Private/Live/GenericCPSCaptureSource.cpp
@@ -0,0 +1,676 @@
+// Copyright Epic Games, Inc. All Rights Reserved.
+
+#include "Live/GenericCPSCaptureSource.h"
+
+#include "Live/GenericCPSCaptureSourceFactory.h"
+#include "GenericCaptureSourceLog.h"
+
+#include "CaptureSourceFrameworkModule.h"
+#include "CaptureSourceDetailsCustomization.h"
+
+#include "Capabilities/RecordingCapability.h"
+#include "Capabilities/IngestCapability.h"
+
+#include "Control/Messages/Constants.h"
+#include "Control/Messages/ControlUpdate.h"
+
+#include "CPS/CPSFileStream.h"
+#include "CPS/CPSDataStream.h"
+
+#include "LiveLinkFaceMetadata.h"
+#include "StereoCameraMetadataParseUtils.h"
+
+#include "CaptureCommunicationManagerModule.h"
+
+#include "HAL/PlatformFileManager.h"
+
+#include "Settings/CaptureManagerSettings.h"
+
+DEFINE_LOG_CATEGORY(LogGenericCPSCaptureSource);
+
+namespace UE::CaptureManager::GenericCPSCaptureSource
+{
+
+static TSharedRef<FCommunicationManager> GetCommunicationManager()
+{
+	FCaptureCommunicationManagerModule& CaptureCommunicationManagerModule =
+		FModuleManager::LoadModuleChecked<FCaptureCommunicationManagerModule>("CaptureCommunicationManager");
+	return CaptureCommunicationManagerModule.GetManager();
+}
+
+static EMediaTexturePixelFormat ConvertPixelFormat(EImagePixelFormat InImagePixelFormat)
+{
+	switch (InImagePixelFormat)
+	{
+		case EImagePixelFormat::U8_RGB:
+			return EMediaTexturePixelFormat::U8_RGB;
+		case EImagePixelFormat::U8_BGR:
+			return EMediaTexturePixelFormat::U8_BGR;
+		case EImagePixelFormat::U8_RGBA:
+			return EMediaTexturePixelFormat::U8_RGBA;
+		case EImagePixelFormat::U8_BGRA:
+			return EMediaTexturePixelFormat::U8_BGRA;
+		case EImagePixelFormat::U8_I444:
+			return EMediaTexturePixelFormat::U8_I444;
+		case EImagePixelFormat::U8_I420:
+			return EMediaTexturePixelFormat::U8_I420;
+		case EImagePixelFormat::U8_NV12:
+			return EMediaTexturePixelFormat::U8_NV12;
+		case EImagePixelFormat::U8_Mono:
+			return EMediaTexturePixelFormat::U8_Mono;
+		case EImagePixelFormat::U16_Mono:
+			return EMediaTexturePixelFormat::U16_Mono;
+		case EImagePixelFormat::F_Mono:
+			return EMediaTexturePixelFormat::F_Mono;
+		case EImagePixelFormat::Undefined:
+		default:
+			check(false);
+			return EMediaTexturePixelFormat::Undefined;
+	}
+}
+
+}
+
+FGenericCPSCaptureSource::FGenericCPSCaptureSource(const FString& InName,
+												   FString InIpAddress,
+												   uint16 InPort)
+	: FCaptureSource(FGenericCPSCaptureSourceFactory::Id, InName)
+	, Device(FCPSDevice::MakeCPSDevice(MoveTemp(InIpAddress), InPort))
+	, CommunicationManager(UE::CaptureManager::GenericCPSCaptureSource::GetCommunicationManager())
+{
+	using namespace UE::CaptureManager;
+
+	TUniquePtr<FRecordingCapability> Recording = MakeUnique<FRecordingCapability>();
+
+	Recording->SetStartRecordingHandler(FCaptureSourceCapability::FCommandHandler::CreateRaw(this, &FGenericCPSCaptureSource::RecordingCommandHandler));
+	Recording->SetStopRecordingHandler(FCaptureSourceCapability::FCommandHandler::CreateRaw(this, &FGenericCPSCaptureSource::RecordingCommandHandler));
+
+	AddCapability(Recording.Get());
+	OwnedCapabilities.Add(FRecordingCapability::Name, MoveTemp(Recording));
+
+	TUniquePtr<FIngestCapability> Ingest = MakeUnique<FIngestCapability>();
+	Ingest->SetImportTakeHandler(FCaptureSourceCapability::FCommandHandler::CreateRaw(this, &FGenericCPSCaptureSource::ImportTakeHandler));
+	Ingest->SetUpdateTakeListHandler(FCaptureSourceCapability::FCommandHandler::CreateRaw(this, &FGenericCPSCaptureSource::UpdateTakeListHandler));
+	Ingest->SetCancelHandler(FCaptureSourceCapability::FCommandHandler::CreateRaw(this, &FGenericCPSCaptureSource::CancelImportHandler));
+
+	AddCapability(Ingest.Get());
+	OwnedCapabilities.Add(FIngestCapability::Name, MoveTemp(Ingest));
+
+	Device->SubscribeToEvent(FConnectionStateChangedEvent::Name, FCaptureEventHandler(
+		FCaptureEventHandler::Type::CreateRaw(this, &FGenericCPSCaptureSource::HandleConnectionChanged),
+		EDelegateExecutionThread::AnyThread));
+
+	Device->SubscribeToEvent(FCPSStateEvent::Name, FCaptureEventHandler(
+		FCaptureEventHandler::Type::CreateRaw(this, &FGenericCPSCaptureSource::HandleCPSStateUpdate),
+		EDelegateExecutionThread::AnyThread));
+
+	Device->SubscribeToEvent(FCPSEvent::Name, FCaptureEventHandler(
+		FCaptureEventHandler::Type::CreateRaw(this, &FGenericCPSCaptureSource::HandleCPSEvent),
+		EDelegateExecutionThread::AnyThread));
+}
+
+FGenericCPSCaptureSource::~FGenericCPSCaptureSource()
+{
+	CommunicationManager->RemoveUploader(GetGuid());
+
+	Stop();
+}
+
+FCaptureVoidResult FGenericCPSCaptureSource::Start()
+{
+	Device->InitiateConnect();
+
+	return MakeValue();
+}
+
+FCaptureVoidResult FGenericCPSCaptureSource::Stop()
+{
+	Device->Stop();
+
+	return MakeValue();
+}
+
+void FGenericCPSCaptureSource::RecordingCommandHandler(TSharedPtr<FCommandBase> InCommand)
+{
+	using namespace UE::CaptureManager;
+
+	if (InCommand->GetName() == FRecordingCapability::StartRecordingCmd)
+	{
+		Device->StartRecording(InCommand->GetParamValue(FRecordingCapability::SlateName).Get<FString>(), 
+							  InCommand->GetParamValue(FRecordingCapability::TakeNumber).Get<int64>());
+	} 
+	else 
+	{
+		Device->StopRecording();
+	}
+}
+
+void FGenericCPSCaptureSource::ImportTakeHandler(TSharedPtr<FCommandBase> InCommand)
+{
+	using namespace UE::CaptureManager;
+
+	TSharedPtr<FImportTakeCmd> ImportTake = StaticCastSharedPtr<FImportTakeCmd>(InCommand);
+
+	FCommandBase::FValues CommandValues = ImportTake->GetParamValues();
+
+	FTakeId TakeId = CommandValues[FIngestCapability::Id].Get<int64>();
+
+	const FGetTakeMetadataResponse::FTakeObject& Take = Device->GetTake(TakeId);
+
+	TSharedPtr<FTaskProgress> TaskProgress = MakeShared<FTaskProgress>
+		(3, FTaskProgress::FProgressReporter::CreateLambda([ImportTake](float InProgress)
+	{
+		FTakeId TakeId = ImportTake->GetParamValues()[FIngestCapability::Id].Get<int64>();
+		ImportTake->ProgressReporter(TakeId, InProgress);
+	}));
+
+	uint64 TotalSize = 0;
+	for (const FGetTakeMetadataResponse::FFileObject& File : Take.Files)
+	{
+		TotalSize += File.Length;
+	}
+
+	const FString DownloadedStorage = GetDefault<UCaptureManagerSettings>()->DownloadDirectory.Path;
+
+	FTaskProgress::FTask Task = TaskProgress->StartTask();
+
+	TUniquePtr<FCPSFileStream> CPSStream = MakeUnique<FCPSFileStream>(DownloadedStorage, TotalSize);
+	CPSStream->SetExportFinished(FCPSFileStream::FExportFinished::CreateRaw(this, &FGenericCPSCaptureSource::OnExportFinished, Take.Name, TakeId, TaskProgress, ImportTake));
+	CPSStream->SetProgressHandler(FCPSFileStream::FReportProgress::CreateRaw(this, &FGenericCPSCaptureSource::OnExportProgressReport, Task));
+
+	Device->StartExport(TakeId, MoveTemp(CPSStream));
+}
+
+void FGenericCPSCaptureSource::UpdateTakeListHandler(TSharedPtr<FCommandBase> InCommand)
+{
+	using namespace UE::CaptureManager;
+
+	TProtocolResult<TArray<FGetTakeMetadataResponse::FTakeObject>> TakesResult = Device->FetchTakeList();
+
+	if (TakesResult.HasError())
+	{
+		UE_LOG(LogGenericCPSCaptureSource, Error, TEXT("Failed to fetch take list: %s"), *TakesResult.StealError().GetMessage());
+		return;
+	}
+
+	TSharedPtr<FUpdateTakeListCmd> Command = StaticCastSharedPtr<FUpdateTakeListCmd>(InCommand);
+
+	TMap<FString, FTakeId> NameToIdMap;
+
+	FIngestCapability* IngestCapability = GetCapability<FIngestCapability>(FIngestCapability::Name);
+	for (const FGetTakeMetadataResponse::FTakeObject& Take : TakesResult.StealValue())
+	{
+		FTakeId TakeId = IngestCapability->AddTakeMetadata(ParseTakeMetadata(Take));
+		Device->AddTakeMetadata(TakeId, Take);
+		NameToIdMap.Add(Take.Name, TakeId);
+	}
+
+	FCPSDataStream::FFileExportFinished OnFileExportFinishedCallback = FCPSDataStream::FFileExportFinished::CreateLambda(
+		[this, TakeNamesMap = MoveTemp(NameToIdMap), IngestCapability](const FString& InTakeName, TProtocolResult<FCPSDataStream::FData> InData)
+	{
+		if (InData.HasError())
+		{
+			UE_LOG(LogGenericCPSCaptureSource, Warning, TEXT("Failed to fetch thumbnail for take: %s"), *InTakeName);
+			return;
+		}
+
+		FCPSDataStream::FData Data = InData.StealValue();
+
+		if (const FTakeId* TakeId = TakeNamesMap.Find(InTakeName))
+		{
+			TOptional<FTakeMetadata> TakeMetadataOpt = IngestCapability->GetTakeMetadata(*TakeId);
+			if (!TakeMetadataOpt.IsSet())
+			{
+				return;
+			}
+
+			FTakeMetadata TakeMetadata = TakeMetadataOpt.GetValue();
+
+			TakeMetadata.Thumbnail = FTakeThumbnailData(MoveTemp(Data));
+
+			IngestCapability->UpdateTakeMetadata(*TakeId, MoveTemp(TakeMetadata));
+
+			IngestCapability->PublishTakeUpdatedEvent(*TakeId);
+		}
+	});
+
+	TUniquePtr<FCPSDataStream> DataStream = MakeUnique<FCPSDataStream>(MoveTemp(OnFileExportFinishedCallback));
+
+	Device->FetchThumbnails(MoveTemp(DataStream));
+
+	Command->Callback(IngestCapability->GetTakeIds());
+}
+
+void FGenericCPSCaptureSource::CancelImportHandler(TSharedPtr<FCommandBase> InCommand)
+{
+	using namespace UE::CaptureManager;
+
+	FTakeId Id = InCommand->GetParamValue(FIngestCapability::Id).Get<int64>();
+
+	CancelImport(Id);
+}
+
+void FGenericCPSCaptureSource::CancelImport(UE::CaptureManager::FTakeId InTakeId)
+{
+	using namespace UE::CaptureManager;
+
+	Device->CancelExport(InTakeId);
+
+	if (TUniquePtr<FCaptureDataConverter>* Converter = Converters.Find(InTakeId))
+	{
+		(*Converter)->Cancel();
+	}
+
+	int32 WorkerId = INDEX_NONE;
+
+	{
+		FScopeLock Lock(&Mutex);
+		if (int32* WorkerIdElem = WorkerForTake.Find(InTakeId))
+		{
+			WorkerId = *WorkerIdElem;
+		}
+	}
+
+	if (WorkerId != INDEX_NONE)
+	{
+		if (TSharedPtr<FUploader> SharedUploader = Uploader.Pin())
+		{
+			SharedUploader->AbortUpload(WorkerId);
+		}
+		else
+		{
+			// If `UploadTakeFinished` is not called (eg. connection between LLH and UE broke), `Uploader` will be deleted, but UI will stay in state where import is in progress 99.9%.
+			// Clicking stop bu tton at that moment will not enter if block above, as uploader does not exist, so in that case we have to manually cleanup the workspace and put UI in proper state.
+			FScopeLock Lock(&Mutex);
+
+			TMap<FGuid, FContext> FilteredMap = ImportContexts.FilterByPredicate([InTakeId](const TPair<FGuid, FContext>& Item) -> bool {
+				return Item.Value.ImportCmd->GetParamValue(FIngestCapability::Id).Get<int64>() == InTakeId;
+				});
+
+			check(FilteredMap.Num() == 1);
+			TMap<FGuid, FContext>::TConstIterator FoundContext = FilteredMap.CreateConstIterator();
+			check(FoundContext);
+
+			const FContext& Context = FoundContext->Value;
+
+			const FString ConversionDirectory = FPaths::Combine(Context.ImportCmd->IngestOptions.Location, 
+																Context.TakeName + TEXT("_") + FoundContext->Key.ToString());
+			IFileManager::Get().DeleteDirectory(*ConversionDirectory, false, true);
+
+			Context.ImportCmd->ProcessFinished(InTakeId, MakeError(FIngestCapabilityError::InternalError, "Aborted by user"));
+
+			ImportContexts.Remove(FoundContext->Key);
+
+			WorkerForTake.Remove(InTakeId);
+		}
+	}
+}
+
+void FGenericCPSCaptureSource::OnExportProgressReport(float InProgress,
+													  UE::CaptureManager::FTaskProgress::FTask InTask)
+{
+	InTask.Update(InProgress);
+}
+
+void FGenericCPSCaptureSource::OnExportFinished(UE::CaptureManager::TProtocolResult<void> InResult,
+												FString InTakeName,
+												UE::CaptureManager::FTakeId InTakeId,
+												TSharedPtr<UE::CaptureManager::FTaskProgress> InTaskProgress,
+												TSharedPtr<UE::CaptureManager::FImportTakeCmd> InCmd)
+{
+	using namespace UE::CaptureManager;
+
+	const FString DownloadedStorage = GetDefault<UCaptureManagerSettings>()->DownloadDirectory.Path;
+	const FString DownloadedTake = FPaths::Combine(DownloadedStorage, InTakeName);
+
+	if (InResult.HasValue())
+	{
+		// Blocking function
+		ConvertData(InTakeName, InTakeId, InTaskProgress, InCmd);
+	}
+	else
+	{
+		TValueOrError<void, FIngestCapabilityError> Result = MakeError(FIngestCapabilityError::ECode::InternalError, InResult.GetError().GetMessage());
+
+		InCmd->ProcessFinished(InTakeId, MoveTemp(Result));
+
+		FScopeLock Lock(&Mutex);
+	}
+
+	IFileManager::Get().DeleteDirectory(*DownloadedTake, false, true);
+}
+
+void FGenericCPSCaptureSource::ConvertData(const FString& InTakeName,
+										   UE::CaptureManager::FTakeId InTakeId,
+										   TSharedPtr<UE::CaptureManager::FTaskProgress> InTaskProgress,
+										   TSharedPtr<UE::CaptureManager::FImportTakeCmd> InImportCmd)
+{
+	using namespace UE::CaptureManager;
+
+	check(InTaskProgress);
+	check(InImportCmd);
+
+	const FString DownloadedStorage = GetDefault<UCaptureManagerSettings>()->DownloadDirectory.Path;
+
+	// Take metadata stored in the capability doesn't have enough information to convert the data
+	// so we need to parse the downloaded take.json
+	TOptional<FTakeMetadata> TakeMetadata = ParseTake(InTakeName);
+
+	if (!TakeMetadata.IsSet())
+	{
+		InImportCmd->ProcessFinished(InTakeId, MakeError(FIngestCapabilityError::InternalError, TEXT("Failed to parse the take metadata")));
+		return;
+	}
+
+	const FString TakeConversionDirectory = InImportCmd->IngestOptions.Location;
+
+	FCaptureDataConverterParams Params;
+	Params.TakeName = InTakeName;
+	Params.TakeMetadata = MoveTemp(TakeMetadata.GetValue());
+	Params.TakeOriginDirectory = FPaths::Combine(DownloadedStorage, InTakeName);
+	Params.TakeOutputDirectory = FPaths::Combine(TakeConversionDirectory, InTakeName + TEXT("_") + GetGuid().ToString());
+
+	FCaptureDataConverterResult<void> Result = RunConversion(InTakeId, InImportCmd->IngestOptions, InTaskProgress, Params);
+
+	Converters.Remove(InTakeId);
+
+	if (Result.HasValue())
+	{
+		FGuid TakeUploadId = FGuid::NewGuid();
+
+		if (!Uploader.IsValid())
+		{
+			Uploader = CommunicationManager->GetUploader(GetGuid(),
+														 GetName(),
+														 FUploader::FUploadStateUpdate::CreateRaw(this, &FGenericCPSCaptureSource::UploadTakeStateHandler),
+														 FUploader::FUploadFinished::CreateRaw(this, &FGenericCPSCaptureSource::UploadTakeFinished));
+		}
+
+		if (TSharedPtr<FUploader> SharedUploader = Uploader.Pin())
+		{
+			FScopeLock Lock(&Mutex);
+			FTaskProgress::FTask UploadTask = InTaskProgress->StartTask();
+
+			ImportContexts.Add(TakeUploadId, { InTakeName, InImportCmd, InTaskProgress, UploadTask });
+
+			int32 Worker = SharedUploader->UploadTake(TakeUploadId, Params.TakeOutputDirectory, Params.TakeMetadata);
+
+			WorkerForTake.Add(InTakeId, Worker);
+		}
+		else
+		{
+			IFileManager::Get().DeleteDirectory(*Params.TakeOutputDirectory, false, true);
+
+			FString Message = TEXT("Failed to upload take: Uploader not found");
+			InImportCmd->ProcessFinished(InTakeId, MakeError(FIngestCapabilityError::InternalError, MoveTemp(Message)));
+		}
+	}
+	else
+	{
+		IFileManager::Get().DeleteDirectory(*Params.TakeOutputDirectory, false, true);
+
+		FCaptureDataConverterError Error = Result.StealError();
+
+		FText Message = FText::Join(FText::FromString(TEXT("\n\r")), Error.GetErrors());
+		InImportCmd->ProcessFinished(InTakeId, MakeError(FIngestCapabilityError::InternalError, Message.ToString()));
+	}
+}
+
+FCaptureDataConverterResult<void> FGenericCPSCaptureSource::RunConversion(UE::CaptureManager::FTakeId InTakeId,
+																		  const UE::CaptureManager::FIngestOptions& InIngestOptions,
+																		  TSharedPtr<UE::CaptureManager::FTaskProgress> InTaskProgress,
+																		  FCaptureDataConverterParams& InParams)
+{
+	using namespace UE::CaptureManager;
+
+	FTaskProgress::FTask ConvertTask = InTaskProgress->StartTask();
+
+	FCaptureConvertVideoOutputParams VideoParams;
+	VideoParams.Format = InIngestOptions.Video.Format;
+	VideoParams.ImageFileName = InIngestOptions.Video.FileNamePrefix;
+	VideoParams.OutputPixelFormat = GenericCPSCaptureSource::ConvertPixelFormat(InIngestOptions.Video.PixelFormat);
+
+	InParams.VideoOutputParams = MoveTemp(VideoParams);
+
+	FCaptureConvertAudioOutputParams AudioParams;
+	AudioParams.Format = InIngestOptions.Audio.Format;
+	AudioParams.AudioFileName = InIngestOptions.Audio.FileNamePrefix;
+
+	InParams.AudioOutputParams = MoveTemp(AudioParams);
+
+	FCaptureConvertDepthOutputParams DepthParams;
+	DepthParams.ImageFileName = TEXT("depth");
+
+	InParams.DepthOutputParams = MoveTemp(DepthParams);
+
+	FCaptureConvertCalibrationOutputParams CalibrationParams;
+	CalibrationParams.FileName = TEXT("calibration");
+
+	InParams.CalibrationOutputParams = MoveTemp(CalibrationParams);
+
+	FCaptureDataConverter::FProgressReporter ProgressReporter =
+		FCaptureDataConverter::FProgressReporter::CreateLambda([this, ConvertTask](const float InProgress) mutable
+	{
+		ConvertTask.Update(InProgress);
+	});
+
+	TUniquePtr<FCaptureDataConverter>& Converter = Converters.FindOrAdd(InTakeId, MakeUnique<FCaptureDataConverter>());
+	return Converter->Run(InParams, MoveTemp(ProgressReporter));
+}
+
+void FGenericCPSCaptureSource::UploadTakeStateHandler(const FGuid& InTakeUploadId, double InProgress)
+{
+	check(InTakeUploadId.IsValid());
+
+	if (!InTakeUploadId.IsValid())
+	{
+		return;
+	}
+
+	FScopeLock Lock(&Mutex);
+	if (FContext* Context = ImportContexts.Find(InTakeUploadId))
+	{
+		Context->Task.Update(InProgress);
+	}
+}
+
+void FGenericCPSCaptureSource::UploadTakeFinished(const FGuid& InTakeUploadId, const FString& InMessage, int32 InCode)
+{
+	using namespace UE::CaptureManager;
+
+	check(InTakeUploadId.IsValid());
+
+	if (!InTakeUploadId.IsValid())
+	{
+		return;
+	}
+
+	FScopeLock Lock(&Mutex);
+	if (const FContext* Context = ImportContexts.Find(InTakeUploadId))
+	{
+		const TSharedPtr<FImportTakeCmd>& ImportTakeCmd = Context->ImportCmd;
+		check(ImportTakeCmd);
+
+		if (!ImportTakeCmd)
+		{
+			return;
+		}
+
+		const bool bShouldCleanConversionDirectory = GetDefault<UCaptureManagerSettings>()->bShouldCleanConversionDirectory;
+		if (bShouldCleanConversionDirectory)
+		{
+			const FString ConversionDirectory = FPaths::Combine(ImportTakeCmd->IngestOptions.Location, 
+																Context->TakeName + TEXT("_") + InTakeUploadId.ToString());
+			IFileManager::Get().DeleteDirectory(*ConversionDirectory, false, true);
+		}
+
+		FTakeId TakeId = ImportTakeCmd->GetParamValue(FIngestCapability::Id).Get<int64>();
+
+		if (InMessage.IsEmpty())
+		{
+			ImportTakeCmd->ProcessFinished(TakeId, MakeValue());
+		}
+		else
+		{
+			ImportTakeCmd->ProcessFinished(TakeId, MakeError(FIngestCapabilityError::InternalError, InMessage));
+		}
+
+
+		ImportContexts.Remove(InTakeUploadId);
+
+		WorkerForTake.Remove(TakeId);
+	}
+}
+
+TOptional<FTakeMetadata> FGenericCPSCaptureSource::ParseTake(const FString& InTakeName)
+{
+	using namespace UE::CaptureManager;
+
+	const FString DownloadedStorage = GetDefault<UCaptureManagerSettings>()->DownloadDirectory.Path;
+	const FString DownloadedTake = FPaths::Combine(DownloadedStorage, InTakeName);
+
+	TArray<FString> TakeFiles;
+	IFileManager::Get().FindFiles(TakeFiles, *DownloadedTake, *FTakeMetadata::FileExtension);
+	if (!TakeFiles.IsEmpty())
+	{
+		check(TakeFiles.Num() == 1);
+
+		TValueOrError<FTakeMetadata, FTakeMetadataParserError> TakeMetadataResult = TakeMetadataParser.Parse(DownloadedTake / TakeFiles[0]);
+		if (!TakeMetadataResult.HasError())
+		{
+			return TakeMetadataResult.StealValue();
+		}
+	}
+
+	TArray<FText> ValidationFailures;
+	UE_LOG(LogGenericCPSCaptureSource, Display, TEXT("Check if the directory (%s) contains a pre UE 5.6 LiveLink take metadata format"), *DownloadedTake);
+	TOptional<FTakeMetadata> ParseOldMetadataResult = LiveLinkMetadata::ParseOldLiveLinkTakeMetadata(DownloadedTake, ValidationFailures);
+
+	if (ParseOldMetadataResult.IsSet())
+	{
+		return ParseOldMetadataResult.GetValue();
+	}
+
+	UE_LOG(LogGenericCPSCaptureSource, Display, TEXT("Check if the directory (%s) contains a pre UE 5.6 StereoCamera take metadata format"), *DownloadedTake);
+	ParseOldMetadataResult = StereoCameraMetadata::ParseOldStereoCameraMetadata(DownloadedTake, ValidationFailures);
+
+	if (ParseOldMetadataResult.IsSet())
+	{
+		return ParseOldMetadataResult.GetValue();
+	}
+
+	UE_LOG(LogGenericCPSCaptureSource, Error, TEXT("Failed to parse take metadata file: %s"), *DownloadedTake);
+	return {};
+}
+
+TArray<FString> FGenericCPSCaptureSource::ExtractFileNames(const UE::CaptureManager::FGetTakeMetadataResponse::FTakeObject& InTake)
+{
+	TArray<FString> FileNames;
+
+	for (const UE::CaptureManager::FGetTakeMetadataResponse::FFileObject& File : InTake.Files)
+	{
+		FileNames.Add(File.Name);
+	}
+
+	return FileNames;
+}
+
+void FGenericCPSCaptureSource::HandleConnectionChanged(TSharedPtr<const UE::CaptureManager::FCaptureEvent> InEvent)
+{
+	using namespace UE::CaptureManager;
+
+	TSharedPtr<const FConnectionStateChangedEvent> Event = StaticCastSharedPtr<const FConnectionStateChangedEvent>(InEvent);
+
+	PublishEvent<UE::CaptureManager::FReachableEvent>(Event->ConnectionState == FConnectionStateChangedEvent::EState::Connected);
+}
+
+void FGenericCPSCaptureSource::HandleCPSStateUpdate(TSharedPtr<const UE::CaptureManager::FCaptureEvent> InEvent)
+{
+	using namespace UE::CaptureManager;
+
+	TSharedPtr<const FCPSStateEvent> Event = StaticCastSharedPtr<const FCPSStateEvent>(InEvent);
+
+	FRecordingCapability* Recording = GetCapability<FRecordingCapability>(FRecordingCapability::Name);
+	Recording->PublishRecordingEvent(Event->GetStateResponse.IsRecording());
+}
+
+void FGenericCPSCaptureSource::HandleCPSEvent(TSharedPtr<const UE::CaptureManager::FCaptureEvent> InEvent)
+{
+	using namespace UE::CaptureManager;
+
+	TSharedPtr<const FCPSEvent> Event = StaticCastSharedPtr<const FCPSEvent>(InEvent);
+
+	if (Event->UpdateMessage->GetAddressPath() == CPS::AddressPaths::GRecordingStatus)
+	{
+		TSharedPtr<FRecordingStatusUpdate> Update = StaticCastSharedPtr<FRecordingStatusUpdate>(Event->UpdateMessage);
+
+		FRecordingCapability* Recording = GetCapability<FRecordingCapability>(FRecordingCapability::Name);
+		Recording->PublishRecordingEvent(Update->IsRecording());
+	}
+	else if (Event->UpdateMessage->GetAddressPath() == CPS::AddressPaths::GTakeAdded)
+	{ 
+		TSharedPtr<const FTakeAddedUpdate> TakeAddedUpdate = StaticCastSharedPtr<const FTakeAddedUpdate>(Event->UpdateMessage);
+
+		TProtocolResult<FGetTakeMetadataResponse::FTakeObject> FetchTakeResult = Device->FetchTake(TakeAddedUpdate->GetTakeName());
+
+		if (FetchTakeResult.IsValid())
+		{
+			FGetTakeMetadataResponse::FTakeObject Take = FetchTakeResult.StealValue();
+
+			FIngestCapability* IngestCapability = GetCapability<FIngestCapability>(FIngestCapability::Name);
+
+			FTakeId TakeId = IngestCapability->AddTakeMetadata(ParseTakeMetadata(Take));
+			Device->AddTakeMetadata(TakeId, MoveTemp(Take));
+
+			IngestCapability->PublishTakeAddedEvent(TakeId);
+
+			UE_LOG(LogGenericCPSCaptureSource, Display, TEXT("Take has been added: %s"), *TakeAddedUpdate->GetTakeName());
+		}
+	}
+	else if (Event->UpdateMessage->GetAddressPath() == CPS::AddressPaths::GTakeRemoved)
+	{
+		TSharedPtr<const FTakeRemovedUpdate> TakeRemoveUpdate = StaticCastSharedPtr<const FTakeRemovedUpdate>(Event->UpdateMessage);
+
+		FTakeId TakeId = Device->GetTakeId(TakeRemoveUpdate->GetTakeName());
+
+		if (TakeId != INDEX_NONE)
+		{
+			CancelImport(TakeId);
+
+			FIngestCapability* IngestCapability = GetCapability<FIngestCapability>(FIngestCapability::Name);
+
+			IngestCapability->RemoveTakeMetadata(TakeId);
+			Device->RemoveTakeMetadata(TakeId);
+
+			IngestCapability->PublishTakeRemovedEvent(TakeId);
+			
+			UE_LOG(LogGenericCPSCaptureSource, Display, TEXT("Take has been removed: %s"), *TakeRemoveUpdate->GetTakeName());
+		}
+	}
+	else
+	{
+		UE_LOG(LogGenericCPSCaptureSource, Display, TEXT("Event arrived that is currently not handled: %s"), *(Event->UpdateMessage->GetAddressPath()));
+	}
+}
+
+FTakeMetadata FGenericCPSCaptureSource::ParseTakeMetadata(const UE::CaptureManager::FGetTakeMetadataResponse::FTakeObject& InTake)
+{
+	FTakeMetadata TakeMetadata;
+
+	TakeMetadata.Slate = InTake.Slate;
+	TakeMetadata.TakeNumber = InTake.TakeNumber;
+	FDateTime DateTime = TakeMetadata.DateTime.Get(FDateTime());
+	FDateTime::ParseIso8601(*InTake.DateTime, DateTime);
+	TakeMetadata.DateTime = DateTime;
+	
+	FTakeMetadata::FVideo Video;
+	Video.FrameRate = InTake.Video.FrameRate;
+	Video.FramesCount = static_cast<uint32>(InTake.Video.Frames);
+	Video.Format = TEXT("mov");
+	Video.FrameHeight = InTake.Video.Height;
+	Video.FrameWidth = InTake.Video.Width;
+
+	TakeMetadata.Video.Add(MoveTemp(Video));
+
+	return TakeMetadata;
+}
\ No newline at end of file
diff --git a/Engine/Plugins/VirtualProduction/CaptureManager/CaptureManagerApp/Source/GenericCaptureSource/Private/Live/GenericCPSCaptureSource.h b/Engine/Plugins/VirtualProduction/CaptureManager/CaptureManagerApp/Source/GenericCaptureSource/Private/Live/GenericCPSCaptureSource.h
new file mode 100644
--- /dev/null
+++ b/Engine/Plugins/VirtualProduction/CaptureManager/CaptureManagerApp/Source/GenericCaptureSource/Private/Live/GenericCPSCaptureSource.h
@@ -0,0 +1,100 @@
+// Copyright Epic Games, Inc. All Rights Reserved.
+
+#pragma once
+
+#include "CaptureSource.h"
+
+#include "Capabilities/IngestCapability.h"
+
+#include "CPS/CPSDevice.h"
+
+#include "Async/StopToken.h"
+#include "Async/TaskProgress.h"
+
+#include "CaptureManagerTakeMetadata.h"
+#include "LiveLinkHubExportManager.h"
+
+#include "CaptureDataConverter.h"
+
+#include "CaptureCommunicationManager.h"
+
+// TODO: Consider having a base class for Generic CPS source and Generic Archive source
+// and use that base class
+class FGenericCPSCaptureSource final : public UE::CaptureManager::FCaptureSource
+{
+public:
+
+	FGenericCPSCaptureSource(const FString& InName,
+							 FString InIpAddress,
+							 uint16 InPort);
+
+	~FGenericCPSCaptureSource();
+
+	virtual FCaptureVoidResult Start() override;
+	virtual FCaptureVoidResult Stop() override;
+
+private:
+
+	void RecordingCommandHandler(TSharedPtr<FCommandBase> InCommand);
+
+	void ImportTakeHandler(TSharedPtr<FCommandBase> InCommand);
+	void UpdateTakeListHandler(TSharedPtr<FCommandBase> InCommand);
+	void CancelImportHandler(TSharedPtr<FCommandBase> InCommand);
+	void CancelImport(UE::CaptureManager::FTakeId InTakeId);
+
+	void OnExportProgressReport(float InProgress,
+								UE::CaptureManager::FTaskProgress::FTask InTask);
+	void OnExportFinished(UE::CaptureManager::TProtocolResult<void> InResult,
+						  FString InTakeName,
+						  UE::CaptureManager::FTakeId InTakeId,
+						  TSharedPtr<UE::CaptureManager::FTaskProgress> InTaskProgress,
+						  TSharedPtr<UE::CaptureManager::FImportTakeCmd> InCmd);
+
+	void ConvertData(const FString& InTakeName, 
+					 UE::CaptureManager::FTakeId InTakeId,
+					 TSharedPtr<UE::CaptureManager::FTaskProgress> InTaskProgress,
+					 TSharedPtr<UE::CaptureManager::FImportTakeCmd> InCmd);
+
+	FCaptureDataConverterResult<void> RunConversion(UE::CaptureManager::FTakeId InTakeId,
+													const UE::CaptureManager::FIngestOptions& InIngestOptions,
+													TSharedPtr<UE::CaptureManager::FTaskProgress> InTaskProgress,
+													FCaptureDataConverterParams& InParams);
+
+	void UploadTakeStateHandler(const FGuid& InTakeUploadId, double InProgress);
+	void UploadTakeFinished(const FGuid& InTakeUploadId, const FString& InMessage, int32 InCode);
+
+	TOptional<FTakeMetadata> ParseTake(const FString& InTakeName);
+	TArray<FString> ExtractFileNames(const UE::CaptureManager::FGetTakeMetadataResponse::FTakeObject& InTake);
+
+	void HandleConnectionChanged(TSharedPtr<const UE::CaptureManager::FCaptureEvent> InEvent);
+	void HandleCPSStateUpdate(TSharedPtr<const UE::CaptureManager::FCaptureEvent> InEvent);
+	void HandleCPSEvent(TSharedPtr<const UE::CaptureManager::FCaptureEvent> InEvent);
+
+	FTakeMetadata ParseTakeMetadata(const UE::CaptureManager::FGetTakeMetadataResponse::FTakeObject& InTake);
+
+	TSharedPtr<FCPSDevice> Device;
+
+	TMap<FString, TUniquePtr<UE::CaptureManager::FCaptureSourceCapability>> OwnedCapabilities;
+
+	uint64 TotalBytes = 0;
+	uint64 RemainingBytes = 0;
+
+	struct FContext
+	{
+		FString TakeName;
+		TSharedPtr<UE::CaptureManager::FImportTakeCmd> ImportCmd;
+		TSharedPtr<UE::CaptureManager::FTaskProgress> TaskProgress;
+		UE::CaptureManager::FTaskProgress::FTask Task;
+	};
+
+	TSharedRef<FCommunicationManager> CommunicationManager;
+	TWeakPtr<FUploader> Uploader;
+
+	TMap<FGuid, FContext> ImportContexts;
+
+	FCriticalSection Mutex;
+	TMap<UE::CaptureManager::FTakeId, int32> WorkerForTake;
+
+	TMap<UE::CaptureManager::FTakeId, TUniquePtr<FCaptureDataConverter>> Converters;
+	FTakeMetadataParser TakeMetadataParser;
+};
\ No newline at end of file
diff --git a/Engine/Plugins/VirtualProduction/CaptureManager/CaptureManagerApp/Source/GenericCaptureSource/Private/Live/GenericCPSCaptureSourceFactory.cpp b/Engine/Plugins/VirtualProduction/CaptureManager/CaptureManagerApp/Source/GenericCaptureSource/Private/Live/GenericCPSCaptureSourceFactory.cpp
new file mode 100644
--- /dev/null
+++ b/Engine/Plugins/VirtualProduction/CaptureManager/CaptureManagerApp/Source/GenericCaptureSource/Private/Live/GenericCPSCaptureSourceFactory.cpp
@@ -0,0 +1,63 @@
+// Copyright Epic Games, Inc. All Rights Reserved.
+
+#include "Live/GenericCPSCaptureSourceFactory.h"
+
+#include "Live/GenericCPSCaptureSource.h"
+
+const FString FGenericCPSCaptureSourceFactory::Id = TEXT("GenericCPSCaptureSource");
+const FText FGenericCPSCaptureSourceFactory::Name = NSLOCTEXT("GenericCPSCaptureSourceFactory", "FactoryName", "Generic CPS");
+
+const FString FGenericCPSCaptureSourceFactory::IpAddress = TEXT("IpAddress");
+const FString FGenericCPSCaptureSourceFactory::Port = TEXT("Port");
+
+FGenericCPSCaptureSourceFactory::FGenericCPSCaptureSourceFactory() = default;
+
+FText FGenericCPSCaptureSourceFactory::GetCaptureSourceFactoryName() const
+{
+	return Name;
+}
+
+FString FGenericCPSCaptureSourceFactory::GetCaptureSourceFactoryId() const
+{
+	return Id;
+}
+
+bool FGenericCPSCaptureSourceFactory::IsDiscoverable() const
+{
+	return true;
+}
+
+void FGenericCPSCaptureSourceFactory::DiscoverImpl(FDiscoverCallback InCallback, FCaptureVoidResult& OutResult)
+{
+}
+
+FGenericCPSCaptureSourceFactory::FCaptureSourceResult
+FGenericCPSCaptureSourceFactory::CreateCaptureSourceImpl(const FString& InName, const TMap<FString, FPropertyValue>& InCreationParamsValue)
+{
+	if (InCreationParamsValue.Num() != 2)
+	{
+		return MakeError(FCaptureSourceError(ECaptureSourceError::MissingArgument));
+	}
+
+	if (!InCreationParamsValue.Contains(IpAddress) || !InCreationParamsValue.Contains(Port))
+	{
+		return MakeError(FCaptureSourceError(ECaptureSourceError::MissingArgument));
+	}
+
+	FString IpAddressValue = InCreationParamsValue[IpAddress].Get<FString>();
+	uint16 PortValue = static_cast<uint16>(InCreationParamsValue[Port].Get<int64>());
+
+	TUniquePtr<FGenericCPSCaptureSource> Source = MakeUnique<FGenericCPSCaptureSource>(InName, IpAddressValue, PortValue);
+
+	Source->SetCreationParams(InCreationParamsValue);
+
+	return MakeValue(MoveTemp(Source));
+}
+
+void FGenericCPSCaptureSourceFactory::CreateCaptureSourceDescriptorImpl(FCaptureSourceDescriptor& OutDescriptor) const
+{
+	OutDescriptor.SetDiscoverable(true);
+
+	OutDescriptor.AddCreationParam(FPropertyDesc(IpAddress, FPropertyDesc::EType::String));
+	OutDescriptor.AddCreationParam(FPropertyDesc(Port, FPropertyDesc::EType::Number));
+}
\ No newline at end of file
diff --git a/Engine/Plugins/VirtualProduction/CaptureManager/CaptureManagerApp/Source/GenericCaptureSource/Private/Live/GenericCPSCaptureSourceFactory.h b/Engine/Plugins/VirtualProduction/CaptureManager/CaptureManagerApp/Source/GenericCaptureSource/Private/Live/GenericCPSCaptureSourceFactory.h
new file mode 100644
--- /dev/null
+++ b/Engine/Plugins/VirtualProduction/CaptureManager/CaptureManagerApp/Source/GenericCaptureSource/Private/Live/GenericCPSCaptureSourceFactory.h
@@ -0,0 +1,30 @@
+// Copyright Epic Games, Inc. All Rights Reserved.
+
+#pragma once
+
+#include "CaptureSourceFactory.h"
+
+class FGenericCPSCaptureSourceFactory final : public FCaptureSourceFactory
+{
+public:
+
+	static const FString Id;
+	static const FText Name;
+
+	static const FString IpAddress;
+	static const FString Port;
+
+	FGenericCPSCaptureSourceFactory();
+
+	virtual FText GetCaptureSourceFactoryName() const override;
+	virtual FString GetCaptureSourceFactoryId() const override;
+
+	virtual bool IsDiscoverable() const override;
+
+private:
+	virtual void DiscoverImpl(FDiscoverCallback InCallback, FCaptureVoidResult& OutResult) override;
+
+	virtual FCaptureSourceResult CreateCaptureSourceImpl(const FString& InName, const TMap<FString, FPropertyValue>& InCreationParamsValue) override;
+	virtual void CreateCaptureSourceDescriptorImpl(FCaptureSourceDescriptor& OutDescriptor) const override;
+
+};
\ No newline at end of file
diff --git a/Engine/Plugins/VirtualProduction/CaptureManager/CaptureManagerApp/Source/GenericLiveLinkDevice/GenericLiveLinkDevice.Build.cs b/Engine/Plugins/VirtualProduction/CaptureManager/CaptureManagerApp/Source/GenericLiveLinkDevice/GenericLiveLinkDevice.Build.cs
new file mode 100644
--- /dev/null
+++ b/Engine/Plugins/VirtualProduction/CaptureManager/CaptureManagerApp/Source/GenericLiveLinkDevice/GenericLiveLinkDevice.Build.cs
@@ -0,0 +1,51 @@
+// Copyright Epic Games, Inc. All Rights Reserved.
+
+using UnrealBuildTool;
+using System.IO;
+
+public class GenericLiveLinkDevice : ModuleRules
+{
+	public GenericLiveLinkDevice(ReadOnlyTargetRules Target) : base(Target)
+	{
+		PCHUsage = PCHUsageMode.UseExplicitOrSharedPCHs;
+		ShortName = "GenericLLD";
+
+		PublicDependencyModuleNames.AddRange(new string[]
+		{
+			"CaptureUtils",
+			"Core"
+		});
+
+		PrivateDependencyModuleNames.AddRange(new string[]
+		{
+			"CaptureProtocolStack",
+			"CaptureDataConverter",
+			"CaptureCommunicationManager",
+			"Engine",
+			"Json",
+			"LiveLinkDevice",
+			"LiveLinkCapabilities",
+			"Media",
+			"LiveLinkFaceMetadata",
+			"LiveLinkHubCaptureMessaging",
+			"LiveLinkHubExportClient",
+			"LiveLinkHub",
+			"ToolWidgets",
+			"Slate",
+			"SlateCore",
+			"StereoCameraMetadata",
+			"CaptureManagerTakeMetadata",
+			"CaptureManagerSettings",
+			"CoreUObject",
+			"UnrealEd"
+		});
+
+		PrivateIncludePaths.AddRange(new string[]
+		{
+			Path.Combine(ModuleDirectory, "Private/Archive"),
+			Path.Combine(ModuleDirectory, "Private/Tests")
+		});
+
+		CppCompileWarningSettings.UnsafeTypeCastWarningLevel = WarningLevel.Error;
+	}
+}
diff --git a/Engine/Plugins/VirtualProduction/CaptureManager/CaptureManagerApp/Source/GenericLiveLinkDevice/Private/Archive/GenericLiveLinkArchiveDevice.cpp b/Engine/Plugins/VirtualProduction/CaptureManager/CaptureManagerApp/Source/GenericLiveLinkDevice/Private/Archive/GenericLiveLinkArchiveDevice.cpp
new file mode 100644
--- /dev/null
+++ b/Engine/Plugins/VirtualProduction/CaptureManager/CaptureManagerApp/Source/GenericLiveLinkDevice/Private/Archive/GenericLiveLinkArchiveDevice.cpp
@@ -0,0 +1,496 @@
+// Copyright Epic Games, Inc. All Rights Reserved.
+
+#include "GenericLiveLinkArchiveDevice.h"
+
+#include "HAL/FileManager.h"
+#include "CaptureCommunicationManagerModule.h"
+#include "Settings/CaptureManagerSettings.h"
+
+#include "Utils/GenericLiveLinkDeviceUtility.h"
+#include "Utils/GenericLiveLinkDeviceLog.h"
+
+#include "LiveLinkFaceMetadata.h"
+#include "StereoCameraMetadataParseUtils.h"
+
+DEFINE_LOG_CATEGORY(LogGenericLiveLinkArchiveDevice);
+
+const UGenericLiveLinkArchiveDeviceSettings* UGenericLiveLinkArchiveDevice::GetSettings() const
+{
+	return CastChecked<UGenericLiveLinkArchiveDeviceSettings>(Settings);
+}
+
+TSubclassOf<ULiveLinkDeviceSettings> UGenericLiveLinkArchiveDevice::StaticGetSettingsClass() const
+{
+	return UGenericLiveLinkArchiveDeviceSettings::StaticClass();
+}
+
+FText UGenericLiveLinkArchiveDevice::GetDisplayName() const
+{
+	return FText::FromString(GetSettings()->DisplayName);
+}
+
+EDeviceHealth UGenericLiveLinkArchiveDevice::GetDeviceHealth() const
+{
+	return EDeviceHealth::Nominal;
+}
+
+FText UGenericLiveLinkArchiveDevice::GetHealthText() const
+{
+	return FText::FromString("Nominal");
+}
+
+void UGenericLiveLinkArchiveDevice::OnDeviceAdded(const FGuid InDeviceGuid, ULiveLinkDeviceSettings* InSettings)
+{
+	Super::OnDeviceAdded(InDeviceGuid, InSettings);
+
+	DeviceGuid = InDeviceGuid;
+
+	CommunicationManager = UE::CaptureManager::Private::GetCommunicationManager();
+
+	UGenericLiveLinkArchiveDeviceSettings* ConfiguredSettings = CastChecked<UGenericLiveLinkArchiveDeviceSettings>(InSettings);
+	FString Path = ConfiguredSettings->StoragePath.Path;
+
+	if (!Path.IsEmpty() && FPaths::DirectoryExists(Path))
+	{
+		SetConnectionStatus(ELiveLinkDeviceConnectionStatus::Connected);
+	}
+	else
+	{
+		SetConnectionStatus(ELiveLinkDeviceConnectionStatus::Disconnected);
+	}
+}
+
+void UGenericLiveLinkArchiveDevice::OnSettingChanged(const FPropertyChangedEvent& InPropertyChangedEvent)
+{
+	Super::OnSettingChanged(InPropertyChangedEvent);
+
+	if (GET_MEMBER_NAME_CHECKED(UGenericLiveLinkArchiveDeviceSettings, StoragePath) == InPropertyChangedEvent.GetMemberPropertyName())
+	{
+		FString Path = GetSettings()->StoragePath.Path;
+		if (!Path.IsEmpty() && FPaths::DirectoryExists(Path))
+		{
+			SetConnectionStatus(ELiveLinkDeviceConnectionStatus::Connected);
+		}
+		else
+		{
+			SetConnectionStatus(ELiveLinkDeviceConnectionStatus::Disconnected);
+		}
+	}
+}
+
+void UGenericLiveLinkArchiveDevice::UpdateTakeListAsync_Implementation()
+{
+	AsyncTask(ENamedThreads::AnyThread, [this]()
+	{
+		RemoveAllTakes();
+
+		FString StoragePath = GetSettings()->StoragePath.Path;
+
+		TArray<FString> TakeMetadataFiles;
+		const bool bIterationResult = IFileManager::Get().IterateDirectoryRecursively(*StoragePath, [&](const TCHAR* InFileNameOrDirectory, bool bInIsDirectory)
+		{
+			if (!bInIsDirectory)
+			{
+				const FString CurrentFileName = FPaths::GetCleanFilename(InFileNameOrDirectory);
+				if (CurrentFileName == TEXT("take.json"))
+				{
+					TakeMetadataFiles.Add(InFileNameOrDirectory);
+				}
+				else if (FPaths::GetExtension(CurrentFileName) == FTakeMetadata::FileExtension)
+				{
+					TakeMetadataFiles.Add(InFileNameOrDirectory);
+				}
+			}
+
+			return true;
+		});
+
+		for (const FString& CurrentTakeFile : TakeMetadataFiles)
+		{
+			TOptional<FTakeMetadata> TakeInfoResult = ReadTake(CurrentTakeFile);
+
+			if (TakeInfoResult.IsSet())
+			{
+				int32 TakeId = AddTake(TakeInfoResult.GetValue());
+				FString CurrentDirectory = FPaths::GetPath(CurrentTakeFile);
+
+				FString TakePath = CurrentDirectory.Right(CurrentDirectory.Len() - StoragePath.Len());
+				FPaths::NormalizeDirectoryName(TakePath); // Removes trailing slash
+				RelativeTakePaths.Add(TakeId, MoveTemp(TakePath));
+			}
+		}
+
+		ExecuteUpdateTakeListCallback(Execute_GetTakeIdentifiers(this));
+	});
+}
+
+void UGenericLiveLinkArchiveDevice::RunIngestTakeAsync(int32 InTakeId, const UIngestCapability_Options* InIngestOptions)
+{
+	AsyncTask(ENamedThreads::AnyThread, [this, 
+			  TakeId = InTakeId, 
+			  IngestOptions = TStrongObjectPtr<const UIngestCapability_Options>(InIngestOptions)]()
+	{
+		using namespace UE::CaptureManager;
+
+		TOptional<FTakeMetadata> MaybeTakeMetadata = GetTakeMetadata(TakeId);
+		if (!MaybeTakeMetadata.IsSet())
+		{
+			FString Message = TEXT("Failed to find the take metadata for the specified take");
+			ExecuteProcessFinishedReporter(TakeId, MakeError(FIngestCapability_Error::InvalidArgument, MoveTemp(Message)));
+			return;
+		}
+
+		FTakeMetadata TakeMetadata = MoveTemp(MaybeTakeMetadata.GetValue());
+
+		FString RelativeTakePath = RelativeTakePaths[TakeId];
+
+		static constexpr uint32 NumberOfTasks = 2; // Convert, Upload
+
+		TSharedPtr<FTaskProgress> TaskProgress 
+			= MakeShared<FTaskProgress>(NumberOfTasks, FTaskProgress::FProgressReporter::CreateLambda([this, TakeId](double InProgress)
+		{
+			ExecuteProcessProgressReporter(TakeId, InProgress);
+		}));
+
+		FString TakeName = TakeMetadata.Slate + TEXT("_") + FString::FromInt(TakeMetadata.TakeNumber);
+
+		const FString TakeConversionDirectory = IngestOptions->Location;
+		const FString StoragePath = GetSettings()->StoragePath.Path;
+
+		FGuid TakeUploadId = FGuid::NewGuid();
+
+		FCaptureDataConverterParams Params;
+		Params.TakeName = TakeName;
+		Params.TakeMetadata = TakeMetadata;
+		Params.TakeOriginDirectory = FPaths::Combine(StoragePath, RelativeTakePath);
+		Params.TakeOutputDirectory = FPaths::Combine(TakeConversionDirectory, TakeName + TEXT("_") + TakeUploadId.ToString());
+
+		FCaptureDataConverterResult<void> Result = RunConversion(TakeId, IngestOptions, TaskProgress, Params);
+
+		Converters.Remove(TakeId);
+
+		if (Result.HasValue())
+		{
+			if (!Uploader.IsValid())
+			{
+				Uploader = CommunicationManager->GetUploader(DeviceGuid,
+															 GetName(),
+															 FUploader::FUploadStateUpdate::CreateUObject(this, &UGenericLiveLinkArchiveDevice::UploadTakeStateHandler),
+															 FUploader::FUploadFinished::CreateUObject(this, &UGenericLiveLinkArchiveDevice::UploadTakeFinished));
+			}
+
+			if (TSharedPtr<FUploader> SharedUploader = Uploader.Pin())
+			{
+				FScopeLock Lock(&Mutex);
+
+				FTaskProgress::FTask UploadTask = TaskProgress->StartTask();
+
+				ImportContexts.Add(TakeUploadId, { TakeName, TakeId, IngestOptions, TaskProgress, UploadTask });
+
+				int32 Worker = SharedUploader->UploadTake(TakeUploadId, Params.TakeOutputDirectory, MoveTemp(TakeMetadata));
+				WorkerForTake.Add(TakeId, Worker);
+			}
+			else
+			{
+				IFileManager::Get().DeleteDirectory(*Params.TakeOutputDirectory, false, true);
+
+				FString Message = TEXT("Failed to upload take: Uploader not found");
+				ExecuteProcessFinishedReporter(TakeId, MakeError(FIngestCapability_Error::UploaderNotFound, MoveTemp(Message)));
+			}
+		}
+		else
+		{
+			IFileManager::Get().DeleteDirectory(*Params.TakeOutputDirectory, false, true);
+
+			FCaptureDataConverterError Error = Result.StealError();
+
+			FText Message = FText::Join(FText::FromString(TEXT("\n\r")), Error.GetErrors());
+			ExecuteProcessFinishedReporter(TakeId, 
+										   MakeError(FIngestCapability_Error::ConversionError, Message.ToString()));
+		}
+	});
+}
+
+void UGenericLiveLinkArchiveDevice::CancelIngest(int32 InTakeId)
+{
+	if (TUniquePtr<FCaptureDataConverter>* Converter = Converters.Find(InTakeId))
+	{
+		(*Converter)->Cancel();
+	}
+
+	int32 WorkerId = INDEX_NONE;
+
+	{
+		FScopeLock Lock(&Mutex);
+		if (int32* WorkerIdElem = WorkerForTake.Find(InTakeId))
+		{
+			WorkerId = *WorkerIdElem;
+		}
+	}
+
+	if (WorkerId != INDEX_NONE)
+	{
+		if (TSharedPtr<FUploader> SharedUploader = Uploader.Pin())
+		{
+			SharedUploader->AbortUpload(WorkerId);
+		}
+		else
+		{
+			// If `UploadTakeFinished` is not called (eg. connection between LLH and UE broke), `Uploader` will be deleted, but UI will stay in state where import is in progress 99.9%.
+			// Clicking stop button at that moment will not enter if block above, as uploader does not exist, so in that case we have to manually cleanup the workspace and put UI in proper state.
+			FScopeLock Lock(&Mutex);
+
+			TMap<FGuid, FContext> FilteredMap = ImportContexts.FilterByPredicate([InTakeId](const TPair<FGuid, FContext>& Item) -> bool
+			{
+				return Item.Value.TakeId == InTakeId;
+			});
+
+			check(FilteredMap.Num() == 1);
+			TMap<FGuid, FContext>::TConstIterator FoundContext = FilteredMap.CreateConstIterator();
+			check(FoundContext);
+
+			const FContext& Context = FoundContext->Value;
+
+			const FString TakeDirectory = FPaths::Combine(Context.IngestOptions->Location,
+														  Context.TakeName + TEXT("_") + FoundContext->Key.ToString());
+			IFileManager::Get().DeleteDirectory(*TakeDirectory, false, true);
+
+			FString Message = TEXT("Aborted by user");
+			ExecuteProcessFinishedReporter(InTakeId, MakeError(FIngestCapability_Error::AbortedByUser, MoveTemp(Message)));
+
+			ImportContexts.Remove(FoundContext->Key);
+
+			WorkerForTake.Remove(InTakeId);
+		}
+	}
+}
+
+ELiveLinkDeviceConnectionStatus UGenericLiveLinkArchiveDevice::GetConnectionStatus_Implementation() const
+{
+	const UGenericLiveLinkArchiveDeviceSettings* DeviceSettings = GetSettings();
+	FString Path = DeviceSettings->StoragePath.Path;
+
+	if (!Path.IsEmpty() && FPaths::DirectoryExists(Path))
+	{
+		return ELiveLinkDeviceConnectionStatus::Connected;
+	}
+
+	return ELiveLinkDeviceConnectionStatus::Disconnected;
+}
+
+FString UGenericLiveLinkArchiveDevice::GetHardwareId_Implementation() const
+{
+	return FPlatformMisc::GetDeviceId();
+}
+
+bool UGenericLiveLinkArchiveDevice::SetHardwareId_Implementation(const FString& HardwareID)
+{
+	return false;
+}
+
+bool UGenericLiveLinkArchiveDevice::Connect_Implementation()
+{
+	return true;
+}
+
+bool UGenericLiveLinkArchiveDevice::Disconnect_Implementation()
+{
+	return true;
+}
+
+FCaptureDataConverterResult<void> UGenericLiveLinkArchiveDevice::RunConversion(int32 InTakeId,
+																			   TStrongObjectPtr<const UIngestCapability_Options> InIngestOptions,
+																			   TSharedPtr<UE::CaptureManager::FTaskProgress> InTaskProgress,
+																			   FCaptureDataConverterParams& InParams)
+{
+	using namespace UE::CaptureManager;
+
+	FTaskProgress::FTask ConvertTask = InTaskProgress->StartTask();
+
+	FCaptureConvertVideoOutputParams VideoParams;
+	VideoParams.Format = InIngestOptions->Video.Format;
+	VideoParams.ImageFileName = InIngestOptions->Video.FileNamePrefix;
+	VideoParams.OutputPixelFormat = Private::ConvertPixelFormat(InIngestOptions->Video.PixelFormat);
+
+	InParams.VideoOutputParams = MoveTemp(VideoParams);
+
+	FCaptureConvertAudioOutputParams AudioParams;
+	AudioParams.Format = InIngestOptions->Audio.Format;
+	AudioParams.AudioFileName = InIngestOptions->Audio.FileNamePrefix;
+
+	InParams.AudioOutputParams = MoveTemp(AudioParams);
+
+	FCaptureConvertDepthOutputParams DepthParams;
+	DepthParams.ImageFileName = TEXT("depth");
+
+	InParams.DepthOutputParams = MoveTemp(DepthParams);
+
+	FCaptureConvertCalibrationOutputParams CalibrationParams;
+	CalibrationParams.FileName = TEXT("calibration");
+
+	InParams.CalibrationOutputParams = MoveTemp(CalibrationParams);
+
+	FCaptureDataConverter::FProgressReporter ProgressReporter =
+		FCaptureDataConverter::FProgressReporter::CreateLambda([ConvertTask](const double InProgress) mutable
+	{
+		ConvertTask.Update(InProgress);
+	});
+
+	TUniquePtr<FCaptureDataConverter>& Converter = Converters.FindOrAdd(InTakeId, MakeUnique<FCaptureDataConverter>());
+	return Converter->Run(InParams, MoveTemp(ProgressReporter));
+}
+
+void UGenericLiveLinkArchiveDevice::UploadTakeStateHandler(const FGuid& InTakeUploadId, double InProgress)
+{
+	if (!ensure(InTakeUploadId.IsValid()))
+	{
+		return;
+	}
+
+	FScopeLock Lock(&Mutex);
+	if (FContext* Context = ImportContexts.Find(InTakeUploadId))
+	{
+		Context->Task.Update(InProgress);
+	}
+}
+
+void UGenericLiveLinkArchiveDevice::UploadTakeFinished(const FGuid& InTakeUploadId, const FString& InMessage, int32 InCode)
+{
+	using namespace UE::CaptureManager;
+
+	if (!ensure(InTakeUploadId.IsValid()))
+	{
+		return;
+	}
+
+	FScopeLock Lock(&Mutex);
+	if (const FContext* Context = ImportContexts.Find(InTakeUploadId))
+	{
+		const bool bShouldCleanConversionDirectory = GetDefault<UCaptureManagerSettings>()->bShouldCleanConversionDirectory;
+		if (bShouldCleanConversionDirectory)
+		{
+			const FString ConversionDirectory = FPaths::Combine(Context->IngestOptions->Location,
+																Context->TakeName + TEXT("_") + InTakeUploadId.ToString());
+			IFileManager::Get().DeleteDirectory(*ConversionDirectory, false, true);
+		}
+
+		if (InMessage.IsEmpty())
+		{
+			ExecuteProcessFinishedReporter(Context->TakeId, MakeValue());
+		}
+		else
+		{
+			ExecuteProcessFinishedReporter(Context->TakeId, MakeError(FIngestCapability_Error::UploaderError, InMessage));
+		}
+
+		ImportContexts.Remove(InTakeUploadId);
+
+		WorkerForTake.Remove(Context->TakeId);
+	}
+}
+
+TOptional<FTakeMetadata> UGenericLiveLinkArchiveDevice::ReadTake(const FString& InCurrentTakeFile)
+{
+	using namespace UE::CaptureManager;
+
+	FTakeMetadataParser TakeMetadataParser;
+	TValueOrError<FTakeMetadata, FTakeMetadataParserError> TakeMetadataResult = TakeMetadataParser.Parse(InCurrentTakeFile);
+	FString CurrentDirectory = FPaths::GetPath(InCurrentTakeFile);
+
+	if (TakeMetadataResult.HasError())
+	{
+		FTakeMetadataParserError ParserError = TakeMetadataResult.StealError();
+		FString OriginName = UE::CaptureManager::Private::ErrorOriginToString(ParserError.Origin);
+
+		FString Message =
+			FString::Format(TEXT("Unable to parse take metadata file - {0} (Error origin: {1}): {2}"),
+							{
+								InCurrentTakeFile,
+								ParserError.Message.ToString(),
+								OriginName
+							});
+
+		UE_LOG(LogGenericLiveLinkArchiveDevice, Warning, TEXT("%s"), *Message);
+
+		UE_LOG(LogGenericLiveLinkArchiveDevice, Display, TEXT("Checking backwards compatible take metadata formats"), *CurrentDirectory);
+
+		TArray<FText> ValidationFailures;
+
+		UE_LOG(LogGenericLiveLinkArchiveDevice, Display, TEXT("Checking directory (%s) for pre UE 5.6 LiveLink take metadata format"), *CurrentDirectory);
+		TOptional<FTakeMetadata> ParseOldMetadataResult = LiveLinkMetadata::ParseOldLiveLinkTakeMetadata(CurrentDirectory, ValidationFailures);
+
+		if (ParseOldMetadataResult.IsSet())
+		{
+			return ParseOldMetadataResult.GetValue();
+		}
+
+		UE_LOG(LogGenericLiveLinkArchiveDevice, Display, TEXT("Checking directory (%s) for pre UE 5.6 StereoCamera take metadata format"), *CurrentDirectory);
+		ParseOldMetadataResult = StereoCameraMetadata::ParseOldStereoCameraMetadata(CurrentDirectory, ValidationFailures);
+
+		if (ParseOldMetadataResult.IsSet())
+		{
+			return ParseOldMetadataResult.GetValue();
+		}
+
+		UE_LOG(LogGenericLiveLinkArchiveDevice, Error, TEXT("Failed to parse take metadata file: %s"), *CurrentDirectory);
+
+		return {};
+	}
+
+	FTakeMetadata& TakeMetadata = TakeMetadataResult.GetValue();
+
+	TOptional<FString> ThumbnailPath = TakeMetadata.Thumbnail.GetThumbnailPath();
+	if (ThumbnailPath.IsSet())
+	{
+		if (FPaths::IsRelative(*ThumbnailPath))
+		{
+			TakeMetadata.Thumbnail = FTakeThumbnailData(FPaths::ConvertRelativePathToFull(CurrentDirectory, *ThumbnailPath));
+		}
+	}
+
+	if (!TakeMetadata.Video.IsEmpty())
+	{
+		for (FTakeMetadata::FVideo& Video : TakeMetadata.Video)
+		{
+			if (FPaths::IsRelative(*Video.Path))
+			{
+				Video.Path = FPaths::ConvertRelativePathToFull(CurrentDirectory, *Video.Path);
+			}
+		}
+	}
+
+	if (!TakeMetadata.Depth.IsEmpty())
+	{
+		for (FTakeMetadata::FVideo& Depth : TakeMetadata.Depth)
+		{
+			if (FPaths::IsRelative(*Depth.Path))
+			{
+				Depth.Path = FPaths::ConvertRelativePathToFull(CurrentDirectory, *Depth.Path);
+			}
+		}
+	}
+
+	if (!TakeMetadata.Audio.IsEmpty())
+	{
+		for (FTakeMetadata::FAudio& Audio : TakeMetadata.Audio)
+		{
+			if (FPaths::IsRelative(*Audio.Path))
+			{
+				Audio.Path = FPaths::ConvertRelativePathToFull(CurrentDirectory, *Audio.Path);
+			}
+		}
+	}
+
+	if (!TakeMetadata.Calibration.IsEmpty())
+	{
+		for (FTakeMetadata::FCalibration& Calib : TakeMetadata.Calibration)
+		{
+			if (FPaths::IsRelative(*Calib.Path))
+			{
+				Calib.Path = FPaths::ConvertRelativePathToFull(CurrentDirectory, *Calib.Path);
+			}
+		}
+	}
+
+	return TakeMetadataResult.StealValue();
+}
\ No newline at end of file
diff --git a/Engine/Plugins/VirtualProduction/CaptureManager/CaptureManagerApp/Source/GenericLiveLinkDevice/Private/Archive/GenericLiveLinkArchiveDevice.h b/Engine/Plugins/VirtualProduction/CaptureManager/CaptureManagerApp/Source/GenericLiveLinkDevice/Private/Archive/GenericLiveLinkArchiveDevice.h
new file mode 100644
--- /dev/null
+++ b/Engine/Plugins/VirtualProduction/CaptureManager/CaptureManagerApp/Source/GenericLiveLinkDevice/Private/Archive/GenericLiveLinkArchiveDevice.h
@@ -0,0 +1,104 @@
+// Copyright Epic Games, Inc. All Rights Reserved.
+
+#pragma once
+
+#include "LiveLinkDevice.h"
+
+#include "Ingest/LiveLinkDeviceCapability_Ingest.h"
+#include "LiveLinkDeviceCapability_Connection.h"
+
+#include "Misc/Optional.h"
+
+#include "Async/TaskProgress.h"
+#include "LiveLinkHubExportManager.h"
+
+#include "CaptureDataConverter.h"
+
+#include "CaptureCommunicationManager.h"
+
+#include "GenericLiveLinkArchiveDevice.generated.h"
+
+
+UCLASS()
+class UGenericLiveLinkArchiveDeviceSettings : public ULiveLinkDeviceSettings
+{
+	GENERATED_BODY()
+
+public:
+	UPROPERTY(EditAnywhere, Category = "Live Link Archive Device")
+	FString DisplayName = TEXT("Live Link Archive Device");
+
+	UPROPERTY(EditAnywhere, Category = "Live Link Archive Device")
+	FDirectoryPath StoragePath;
+};
+
+
+UCLASS()
+class UGenericLiveLinkArchiveDevice : public ULiveLinkDevice
+	, public ILiveLinkDeviceCapability_Ingest
+	, public ILiveLinkDeviceCapability_Connection
+{
+	GENERATED_BODY()
+
+public:
+	const UGenericLiveLinkArchiveDeviceSettings* GetSettings() const;
+
+	//~ Begin ULiveLinkDevice interface
+	virtual TSubclassOf<ULiveLinkDeviceSettings> StaticGetSettingsClass() const override;
+	virtual FText GetDisplayName() const override;
+	virtual EDeviceHealth GetDeviceHealth() const override;
+	virtual FText GetHealthText() const override;
+
+	virtual void OnDeviceAdded(const FGuid InDeviceGuid, ULiveLinkDeviceSettings* InSettings) override;
+	virtual void OnSettingChanged(const FPropertyChangedEvent& InPropertyChangedEvent) override;
+	//~ End ULiveLinkDevice interface
+
+private:
+	//~ Begin ULiveLinkDeviceCapability_Ingest interface
+	virtual void UpdateTakeListAsync_Implementation() override;
+	virtual void RunIngestTakeAsync(int32 InTakeId,
+									const UIngestCapability_Options* InIngestOptions) override;
+	virtual void CancelIngest(int32 InTakeId) override;
+	//~ End ULiveLinkDeviceCapability_Ingest interface
+
+	//~ Begin ILiveLinkDeviceCapability_Connection interface
+	virtual ELiveLinkDeviceConnectionStatus GetConnectionStatus_Implementation() const override;
+	virtual FString GetHardwareId_Implementation() const override;
+	virtual bool SetHardwareId_Implementation(const FString& HardwareID) override;
+	virtual bool Connect_Implementation() override;
+	virtual bool Disconnect_Implementation() override;
+	//~ End ILiveLinkDeviceCapability_Connection interface
+
+	static TOptional<FTakeMetadata> ReadTake(const FString& InCurrentTakeFile);
+
+	FCaptureDataConverterResult<void> RunConversion(int32 InTakeId,
+													TStrongObjectPtr<const UIngestCapability_Options> InIngestOptions,
+													TSharedPtr<UE::CaptureManager::FTaskProgress> InTaskProgress,
+													FCaptureDataConverterParams& InParams);
+
+	void UploadTakeStateHandler(const FGuid& InTakeUploadId, double InProgress);
+	void UploadTakeFinished(const FGuid& InTakeUploadId, const FString& InMessage, int32 InCode);
+
+	struct FContext
+	{
+		FString TakeName;
+		int32 TakeId = INDEX_NONE;
+		TStrongObjectPtr<const UIngestCapability_Options> IngestOptions;
+		TSharedPtr<UE::CaptureManager::FTaskProgress> TaskProgress;
+		UE::CaptureManager::FTaskProgress::FTask Task;
+	};
+
+	TSharedPtr<FCommunicationManager> CommunicationManager;
+	TWeakPtr<FUploader> Uploader;
+
+	TMap<FGuid, FContext> ImportContexts;
+	TMap<int32, FString> RelativeTakePaths;
+
+	FCriticalSection Mutex;
+	TMap<int32, int32> WorkerForTake;
+
+	TMap<int32, TUniquePtr<FCaptureDataConverter>> Converters;
+	FTakeMetadataParser TakeMetadataParser;
+
+	FGuid DeviceGuid;
+};
