CaptureManager: Move plugin to Virtual Production directory

diff --git a/Engine/Plugins/VirtualProduction/CaptureManager/CaptureManagerApp/Source/GenericLiveLinkDevice/Private/CPS/GenericLiveLinkCPSDevice.cpp b/Engine/Plugins/VirtualProduction/CaptureManager/CaptureManagerApp/Source/GenericLiveLinkDevice/Private/CPS/GenericLiveLinkCPSDevice.cpp
new file mode 100644
--- /dev/null
+++ b/Engine/Plugins/VirtualProduction/CaptureManager/CaptureManagerApp/Source/GenericLiveLinkDevice/Private/CPS/GenericLiveLinkCPSDevice.cpp
@@ -0,0 +1,694 @@
+// Copyright Epic Games, Inc. All Rights Reserved.
+
+#include "GenericLiveLinkCPSDevice.h"
+
+#include "CaptureCommunicationManagerModule.h"
+#include "Settings/CaptureManagerSettings.h"
+
+#include "ILiveLinkRecordingSessionInfo.h"
+
+#include "Control/Messages/Constants.h"
+#include "Control/Messages/ControlUpdate.h"
+
+#include "CPS/Protocol/CPSDataStream.h"
+#include "CPS/Protocol/CPSFileStream.h"
+
+#include "LiveLinkFaceMetadata.h"
+#include "StereoCameraMetadataParseUtils.h"
+
+#include "Utils/GenericLiveLinkDeviceUtility.h"
+
+UGenericLiveLinkCPSDeviceSettings* UGenericLiveLinkCPSDevice::GetSettings()
+{
+	return CastChecked<UGenericLiveLinkCPSDeviceSettings>(Settings);
+}
+const UGenericLiveLinkCPSDeviceSettings* UGenericLiveLinkCPSDevice::GetSettings() const
+{
+	return CastChecked<UGenericLiveLinkCPSDeviceSettings>(Settings);
+}
+
+TSubclassOf<ULiveLinkDeviceSettings> UGenericLiveLinkCPSDevice::StaticGetSettingsClass() const
+{
+	return UGenericLiveLinkCPSDeviceSettings::StaticClass();
+}
+
+FText UGenericLiveLinkCPSDevice::GetDisplayName() const
+{
+	return FText::FromString(GetSettings()->DisplayName);
+}
+
+EDeviceHealth UGenericLiveLinkCPSDevice::GetDeviceHealth() const
+{
+	return EDeviceHealth::Nominal;
+}
+
+FText UGenericLiveLinkCPSDevice::GetHealthText() const
+{
+	return FText::FromString("Example Health");
+}
+
+void UGenericLiveLinkCPSDevice::OnDeviceAdded(const FGuid InDeviceGuid, ULiveLinkDeviceSettings* InSettings)
+{
+	CastChecked<UGenericLiveLinkCPSDeviceSettings>(InSettings)->ConnectAction.DeviceGuid = InDeviceGuid;
+
+	DeviceGuid = InDeviceGuid;
+
+	CommunicationManager = UE::CaptureManager::Private::GetCommunicationManager();
+
+	Super::OnDeviceAdded(InDeviceGuid, InSettings);
+}
+
+void UGenericLiveLinkCPSDevice::OnDeviceRemoved()
+{
+	RemoveAllTakes();
+
+	ILiveLinkDeviceCapability_Connection::Execute_Disconnect(this);
+
+	Super::OnDeviceRemoved();
+}
+
+void UGenericLiveLinkCPSDevice::OnSettingChanged(const FPropertyChangedEvent& InPropertyChangedEvent)
+{
+	Super::OnSettingChanged(InPropertyChangedEvent);
+}
+
+void UGenericLiveLinkCPSDevice::UpdateTakeListAsync_Implementation()
+{
+	using namespace UE::CaptureManager;
+
+	RemoveAllTakes();
+
+	TProtocolResult<TArray<FGetTakeMetadataResponse::FTakeObject>> TakesResult = Device->FetchTakeList();
+
+	if (TakesResult.HasError())
+	{
+		return;
+	}
+
+	TMap<FString, FTakeId> NameToIdMap;
+
+	for (const FGetTakeMetadataResponse::FTakeObject& Take : TakesResult.StealValue())
+	{
+		FTakeId TakeId = AddTake(ParseTakeMetadata(Take));
+		Device->AddTakeMetadata(TakeId, Take);
+		NameToIdMap.Add(Take.Name, TakeId);
+	}
+
+	FetchThumbnails(MoveTemp(NameToIdMap));
+
+	ExecuteUpdateTakeListCallback(Execute_GetTakeIdentifiers(this));
+}
+
+void UGenericLiveLinkCPSDevice::RunIngestTakeAsync(int32 InTakeId, const UIngestCapability_Options* InIngestOptions)
+{
+	using namespace UE::CaptureManager;
+
+	const FGetTakeMetadataResponse::FTakeObject& Take = Device->GetTake(InTakeId);
+
+	static constexpr uint32 NumberOfTasks = 3; // Download, Convert, Upload
+
+	TSharedPtr<FTaskProgress> TaskProgress = MakeShared<FTaskProgress>
+		(NumberOfTasks, FTaskProgress::FProgressReporter::CreateLambda([this, TakeId = InTakeId](double InProgress)
+	{
+		ExecuteProcessProgressReporter(TakeId, InProgress);
+	}));
+
+	uint64 TotalSize = 0;
+	for (const FGetTakeMetadataResponse::FFileObject& File : Take.Files)
+	{
+		TotalSize += File.Length;
+	}
+
+	const FString DownloadedStorage = GetDefault<UCaptureManagerSettings>()->DownloadDirectory.Path;
+
+	FTaskProgress::FTask Task = TaskProgress->StartTask();
+
+	TStrongObjectPtr<const UIngestCapability_Options> IngestOptions(InIngestOptions);
+
+	TUniquePtr<FCPSFileStream> CPSStream = MakeUnique<FCPSFileStream>(DownloadedStorage, TotalSize);
+	CPSStream->SetExportFinished(FCPSFileStream::FExportFinished::CreateUObject(this, &UGenericLiveLinkCPSDevice::OnExportFinished, Take.Name, InTakeId, TaskProgress, IngestOptions));
+	CPSStream->SetProgressHandler(FCPSFileStream::FReportProgress::CreateUObject(this, &UGenericLiveLinkCPSDevice::OnExportProgressReport, Task));
+
+	Device->StartExport(InTakeId, MoveTemp(CPSStream));
+}
+
+void UGenericLiveLinkCPSDevice::CancelIngest(int32 InTakeId)
+{
+	Device->CancelExport(InTakeId);
+
+	if (TUniquePtr<FCaptureDataConverter>* Converter = Converters.Find(InTakeId))
+	{
+		(*Converter)->Cancel();
+	}
+
+	int32 WorkerId = INDEX_NONE;
+
+	{
+		FScopeLock Lock(&Mutex);
+		if (int32* WorkerIdElem = WorkerForTake.Find(InTakeId))
+		{
+			WorkerId = *WorkerIdElem;
+		}
+	}
+
+	if (WorkerId != INDEX_NONE)
+	{
+		if (TSharedPtr<FUploader> SharedUploader = Uploader.Pin())
+		{
+			SharedUploader->AbortUpload(WorkerId);
+		}
+		else
+		{
+			// If `UploadTakeFinished` is not called (eg. connection between LLH and UE broke), `Uploader` will be deleted, but UI will stay in state where import is in progress 99.9%.
+			// Clicking stop button at that moment will not enter if block above, as uploader does not exist, so in that case we have to manually cleanup the workspace and put UI in proper state.
+			FScopeLock Lock(&Mutex);
+
+			TMap<FGuid, FContext> FilteredMap = ImportContexts.FilterByPredicate([InTakeId](const TPair<FGuid, FContext>& Item) -> bool
+			{
+				return Item.Value.TakeId == InTakeId;
+			});
+
+			check(FilteredMap.Num() == 1);
+			TMap<FGuid, FContext>::TConstIterator FoundContext = FilteredMap.CreateConstIterator();
+			check(FoundContext);
+
+			const FContext& Context = FoundContext->Value;
+
+			const FString TakeDirectory = FPaths::Combine(Context.IngestOptions->Location,
+														  Context.TakeName + TEXT("_") + FoundContext->Key.ToString());
+			IFileManager::Get().DeleteDirectory(*TakeDirectory, false, true);
+
+			FString Message = TEXT("Aborted by user");
+			ExecuteProcessFinishedReporter(InTakeId, MakeError(FIngestCapability_Error::AbortedByUser, MoveTemp(Message)));
+
+			ImportContexts.Remove(FoundContext->Key);
+
+			WorkerForTake.Remove(InTakeId);
+		}
+	}
+}
+
+ELiveLinkDeviceConnectionStatus UGenericLiveLinkCPSDevice::GetConnectionStatus_Implementation() const
+{
+	if (!Device)
+	{
+		return ELiveLinkDeviceConnectionStatus::Disconnected;
+	}
+
+	if (bIsConnecting)
+	{
+		return ELiveLinkDeviceConnectionStatus::Connecting;
+	}
+	else if (Device->IsConnected())
+	{
+		return ELiveLinkDeviceConnectionStatus::Connected;
+	}
+
+	return ELiveLinkDeviceConnectionStatus::Disconnected;
+}
+
+FString UGenericLiveLinkCPSDevice::GetHardwareId_Implementation() const
+{
+	return GetSettings()->IpAddress.IpAddressString;
+}
+
+bool UGenericLiveLinkCPSDevice::SetHardwareId_Implementation(const FString& HardwareID)
+{
+	return false;
+}
+
+bool UGenericLiveLinkCPSDevice::Connect_Implementation()
+{
+	UGenericLiveLinkCPSDeviceSettings* DeviceSettings = GetSettings();
+
+	using namespace UE::CaptureManager;
+
+	Device = FCPSDevice::MakeCPSDevice(DeviceSettings->IpAddress.IpAddressString, DeviceSettings->Port);
+	
+	Device->SubscribeToEvent(FConnectionStateChangedEvent::Name, FCaptureEventHandler(
+		FCaptureEventHandler::Type::CreateUObject(this, &UGenericLiveLinkCPSDevice::HandleConnectionChanged),
+		EDelegateExecutionThread::AnyThread));
+
+	Device->SubscribeToEvent(FCPSStateEvent::Name, FCaptureEventHandler(
+		FCaptureEventHandler::Type::CreateUObject(this, &UGenericLiveLinkCPSDevice::HandleCPSStateUpdate),
+		EDelegateExecutionThread::AnyThread));
+
+	Device->SubscribeToEvent(FCPSEvent::Name, FCaptureEventHandler(
+		FCaptureEventHandler::Type::CreateUObject(this, &UGenericLiveLinkCPSDevice::HandleCPSEvent),
+		EDelegateExecutionThread::AnyThread));
+
+	Device->InitiateConnect();
+
+	return true;
+}
+
+bool UGenericLiveLinkCPSDevice::Disconnect_Implementation()
+{
+	AsyncTask(ENamedThreads::AnyThread, [this]()
+	{
+		if (Device)
+		{
+			Device->Stop();
+			Device->UnsubscribeAll();
+			bIsConnecting = false;
+			Device = nullptr;
+		}
+	});
+	
+	return true;
+}
+
+bool UGenericLiveLinkCPSDevice::StartRecording_Implementation()
+{
+	if (!Device)
+	{
+		return false;
+	}
+
+	using namespace UE::CaptureManager;
+
+	ILiveLinkRecordingSessionInfo& SessionInfo = ILiveLinkRecordingSessionInfo::Get();
+	if (SessionInfo.GetSlateName().IsEmpty() || SessionInfo.GetTakeNumber() == -1)
+	{
+		return false;
+	}
+
+	uint16 TakeNumber = static_cast<uint16>(SessionInfo.GetTakeNumber());
+	TProtocolResult<void> Result = Device->StartRecording(SessionInfo.GetSlateName(), TakeNumber);
+
+	return Result.HasValue();
+}
+
+bool UGenericLiveLinkCPSDevice::StopRecording_Implementation()
+{
+	if (!Device)
+	{
+		return false;
+	}
+
+	using namespace UE::CaptureManager;
+
+	TProtocolResult<void> Result = Device->StopRecording();
+
+	return Result.HasValue();
+}
+
+bool UGenericLiveLinkCPSDevice::IsRecording_Implementation() const
+{
+	return bIsRecording;
+}
+
+void UGenericLiveLinkCPSDevice::HandleConnectionChanged(TSharedPtr<const UE::CaptureManager::FCaptureEvent> InEvent)
+{
+	using namespace UE::CaptureManager;
+
+	TSharedPtr<const FConnectionStateChangedEvent> Event = StaticCastSharedPtr<const FConnectionStateChangedEvent>(InEvent);
+
+	if (Event->ConnectionState == FConnectionStateChangedEvent::EState::Connecting)
+	{
+		bIsConnecting = true;
+	}
+	else if (Event->ConnectionState == FConnectionStateChangedEvent::EState::Connected)
+	{
+		bIsConnecting = false;
+	}
+
+	auto GetConnectionState = [](FConnectionStateChangedEvent::EState InConnectionState) -> ELiveLinkDeviceConnectionStatus
+	{
+		switch (InConnectionState)
+		{
+			case FConnectionStateChangedEvent::EState::Connecting:
+				return ELiveLinkDeviceConnectionStatus::Connecting;
+			case FConnectionStateChangedEvent::EState::Connected:
+				return ELiveLinkDeviceConnectionStatus::Connected;
+			case FConnectionStateChangedEvent::EState::Disconnected:
+			case FConnectionStateChangedEvent::EState::Unknown:
+			default:
+				return ELiveLinkDeviceConnectionStatus::Disconnected;
+		}
+	};
+
+	ELiveLinkDeviceConnectionStatus ConnectionStatus = GetConnectionState(Event->ConnectionState);
+
+	SetConnectionStatus(ConnectionStatus);
+}
+
+void UGenericLiveLinkCPSDevice::HandleCPSStateUpdate(TSharedPtr<const UE::CaptureManager::FCaptureEvent> InEvent)
+{
+	using namespace UE::CaptureManager;
+
+	TSharedPtr<const FCPSStateEvent> Event = StaticCastSharedPtr<const FCPSStateEvent>(InEvent);
+
+	bIsRecording = Event->GetStateResponse.IsRecording();
+}
+
+void UGenericLiveLinkCPSDevice::HandleCPSEvent(TSharedPtr<const UE::CaptureManager::FCaptureEvent> InEvent)
+{
+	using namespace UE::CaptureManager;
+
+	TSharedPtr<const FCPSEvent> Event = StaticCastSharedPtr<const FCPSEvent>(InEvent);
+
+	if (Event->UpdateMessage->GetAddressPath() == CPS::AddressPaths::GRecordingStatus)
+	{
+		TSharedPtr<FRecordingStatusUpdate> Update = StaticCastSharedPtr<FRecordingStatusUpdate>(Event->UpdateMessage);
+		bIsRecording = Update->IsRecording();
+	}
+	else if (Event->UpdateMessage->GetAddressPath() == CPS::AddressPaths::GTakeAdded)
+	{
+		TSharedPtr<const FTakeAddedUpdate> TakeAddedUpdate = StaticCastSharedPtr<const FTakeAddedUpdate>(Event->UpdateMessage);
+
+		TProtocolResult<FGetTakeMetadataResponse::FTakeObject> FetchTakeResult = Device->FetchTake(TakeAddedUpdate->GetTakeName());
+
+		if (FetchTakeResult.IsValid())
+		{
+			FGetTakeMetadataResponse::FTakeObject Take = FetchTakeResult.StealValue();
+			
+			FString TakeName = Take.Name;
+
+			FTakeId TakeId = AddTake(ParseTakeMetadata(Take));
+			Device->AddTakeMetadata(TakeId, MoveTemp(Take));
+
+			PublishEvent<FIngestCapability_TakeAddedEvent>(TakeId);
+
+			FetchThumbnail(TakeName, TakeId);
+		}
+	}
+	else if (Event->UpdateMessage->GetAddressPath() == CPS::AddressPaths::GTakeRemoved)
+	{
+		TSharedPtr<const FTakeRemovedUpdate> TakeRemoveUpdate = StaticCastSharedPtr<const FTakeRemovedUpdate>(Event->UpdateMessage);
+
+		FTakeId TakeId = Device->GetTakeId(TakeRemoveUpdate->GetTakeName());
+
+		if (TakeId != INDEX_NONE)
+		{
+			CancelIngest(TakeId);
+
+			RemoveTake(TakeId);
+			Device->RemoveTakeMetadata(TakeId);
+			
+			PublishEvent<FIngestCapability_TakeRemovedEvent>(TakeId);
+		}
+	}
+}
+
+FTakeMetadata UGenericLiveLinkCPSDevice::ParseTakeMetadata(const UE::CaptureManager::FGetTakeMetadataResponse::FTakeObject& InTake)
+{
+	FTakeMetadata TakeMetadata;
+
+	TakeMetadata.Slate = InTake.Slate;
+	TakeMetadata.TakeNumber = InTake.TakeNumber;
+	FDateTime DateTime = TakeMetadata.DateTime.Get(FDateTime());
+	FDateTime::ParseIso8601(*InTake.DateTime, DateTime);
+	TakeMetadata.DateTime = DateTime;
+
+	FTakeMetadata::FVideo Video;
+	Video.FrameRate = InTake.Video.FrameRate;
+	Video.FramesCount = static_cast<uint32>(InTake.Video.Frames);
+	Video.Format = TEXT("mov");
+	Video.FrameHeight = InTake.Video.Height;
+	Video.FrameWidth = InTake.Video.Width;
+
+	TakeMetadata.Video.Add(MoveTemp(Video));
+
+	return TakeMetadata;
+}
+
+void UGenericLiveLinkCPSDevice::OnExportProgressReport(float InProgress,
+													   UE::CaptureManager::FTaskProgress::FTask InTask)
+{
+	InTask.Update(static_cast<double>(InProgress));
+}
+
+void UGenericLiveLinkCPSDevice::OnExportFinished(UE::CaptureManager::TProtocolResult<void> InResult,
+												 FString InTakeName,
+												 UE::CaptureManager::FTakeId InTakeId,
+												 TSharedPtr<UE::CaptureManager::FTaskProgress> InTaskProgress,
+												 TStrongObjectPtr<const UIngestCapability_Options> InIngestOptions)
+{
+	using namespace UE::CaptureManager;
+
+	const FString DownloadedStorage = GetDefault<UCaptureManagerSettings>()->DownloadDirectory.Path;
+	const FString DownloadedTake = FPaths::Combine(DownloadedStorage, InTakeName);
+
+	if (InResult.HasValue())
+	{
+		// Blocking function
+		ConvertData(InTakeName, InTakeId, InTaskProgress, InIngestOptions);
+	}
+	else
+	{
+		TValueOrError<void, FIngestCapability_Error> Result = MakeError(FIngestCapability_Error::DownloaderError, InResult.GetError().GetMessage());
+
+		ExecuteProcessFinishedReporter(InTakeId, MoveTemp(Result));
+	}
+
+	IFileManager::Get().DeleteDirectory(*DownloadedTake, false, true);
+}
+
+void UGenericLiveLinkCPSDevice::ConvertData(const FString& InTakeName,
+											UE::CaptureManager::FTakeId InTakeId,
+											TSharedPtr<UE::CaptureManager::FTaskProgress> InTaskProgress,
+											TStrongObjectPtr<const UIngestCapability_Options> InIngestOptions)
+{
+	using namespace UE::CaptureManager;
+
+	check(InTaskProgress);
+	check(InIngestOptions);
+
+	const FString DownloadedStorage = GetDefault<UCaptureManagerSettings>()->DownloadDirectory.Path;
+
+	// Take metadata stored in the capability doesn't have enough information to convert the data
+	// so we need to parse the downloaded take.json
+	TOptional<FTakeMetadata> TakeMetadata = ParseTake(InTakeName);
+
+	if (!TakeMetadata.IsSet())
+	{
+		ExecuteProcessFinishedReporter(InTakeId, MakeError(FIngestCapability_Error::InternalError, TEXT("Failed to parse the take metadata")));
+		return;
+	}
+
+	const FString TakeConversionDirectory = InIngestOptions->Location;
+	FGuid TakeUploadId = FGuid::NewGuid();
+
+	FCaptureDataConverterParams Params;
+	Params.TakeName = InTakeName;
+	Params.TakeMetadata = MoveTemp(TakeMetadata.GetValue());
+	Params.TakeOriginDirectory = FPaths::Combine(DownloadedStorage, InTakeName);
+	Params.TakeOutputDirectory = FPaths::Combine(TakeConversionDirectory, InTakeName + TEXT("_") + TakeUploadId.ToString());
+
+	FCaptureDataConverterResult<void> Result = RunConversion(InTakeId, InIngestOptions, InTaskProgress, Params);
+
+	Converters.Remove(InTakeId);
+
+	if (Result.HasValue())
+	{
+		if (!Uploader.IsValid())
+		{
+			Uploader = CommunicationManager->GetUploader(DeviceGuid,
+														 GetName(),
+														 FUploader::FUploadStateUpdate::CreateUObject(this, &UGenericLiveLinkCPSDevice::UploadTakeStateHandler),
+														 FUploader::FUploadFinished::CreateUObject(this, &UGenericLiveLinkCPSDevice::UploadTakeFinished));
+		}
+
+		if (TSharedPtr<FUploader> SharedUploader = Uploader.Pin())
+		{
+			FScopeLock Lock(&Mutex);
+			FTaskProgress::FTask UploadTask = InTaskProgress->StartTask();
+
+			ImportContexts.Add(TakeUploadId, { InTakeName, InTakeId, InIngestOptions, InTaskProgress, UploadTask });
+
+			int32 Worker = SharedUploader->UploadTake(TakeUploadId, Params.TakeOutputDirectory, Params.TakeMetadata);
+
+			WorkerForTake.Add(InTakeId, Worker);
+		}
+		else
+		{
+			IFileManager::Get().DeleteDirectory(*Params.TakeOutputDirectory, false, true);
+
+			FString Message = TEXT("Failed to upload take: Uploader not found");
+			ExecuteProcessFinishedReporter(InTakeId, MakeError(FIngestCapability_Error::UploaderNotFound, MoveTemp(Message)));
+		}
+	}
+	else
+	{
+		FCaptureDataConverterError Error = Result.StealError();
+
+		FText Message = FText::Join(FText::FromString(TEXT("\n\r")), Error.GetErrors());
+		ExecuteProcessFinishedReporter(InTakeId,
+									   MakeError(FIngestCapability_Error::ConversionError, Message.ToString()));
+	}
+}
+
+FCaptureDataConverterResult<void> UGenericLiveLinkCPSDevice::RunConversion(int32 InTakeId,
+																		   TStrongObjectPtr<const UIngestCapability_Options> InIngestOptions,
+																		   TSharedPtr<UE::CaptureManager::FTaskProgress> InTaskProgress,
+																		   FCaptureDataConverterParams& InParams)
+{
+	using namespace UE::CaptureManager;
+
+	FTaskProgress::FTask ConvertTask = InTaskProgress->StartTask();
+
+	FCaptureConvertVideoOutputParams VideoParams;
+	VideoParams.Format = InIngestOptions->Video.Format;
+	VideoParams.ImageFileName = InIngestOptions->Video.FileNamePrefix;
+	VideoParams.OutputPixelFormat = Private::ConvertPixelFormat(InIngestOptions->Video.PixelFormat);
+
+	InParams.VideoOutputParams = MoveTemp(VideoParams);
+
+	FCaptureConvertAudioOutputParams AudioParams;
+	AudioParams.Format = InIngestOptions->Audio.Format;
+	AudioParams.AudioFileName = InIngestOptions->Audio.FileNamePrefix;
+
+	InParams.AudioOutputParams = MoveTemp(AudioParams);
+
+	FCaptureConvertDepthOutputParams DepthParams;
+	DepthParams.ImageFileName = TEXT("depth");
+
+	InParams.DepthOutputParams = MoveTemp(DepthParams);
+
+	FCaptureConvertCalibrationOutputParams CalibrationParams;
+	CalibrationParams.FileName = TEXT("calibration");
+
+	InParams.CalibrationOutputParams = MoveTemp(CalibrationParams);
+
+	FCaptureDataConverter::FProgressReporter ProgressReporter =
+		FCaptureDataConverter::FProgressReporter::CreateLambda([ConvertTask](const double InProgress) mutable
+															   {
+																   ConvertTask.Update(InProgress);
+															   });
+
+	TUniquePtr<FCaptureDataConverter>& Converter = Converters.FindOrAdd(InTakeId, MakeUnique<FCaptureDataConverter>());
+	return Converter->Run(InParams, MoveTemp(ProgressReporter));
+}
+
+void UGenericLiveLinkCPSDevice::UploadTakeStateHandler(const FGuid& InTakeUploadId, double InProgress)
+{
+	if (!ensure(InTakeUploadId.IsValid()))
+	{
+		return;
+	}
+
+	FScopeLock Lock(&Mutex);
+	if (FContext* Context = ImportContexts.Find(InTakeUploadId))
+	{
+		Context->Task.Update(InProgress);
+	}
+}
+
+void UGenericLiveLinkCPSDevice::UploadTakeFinished(const FGuid& InTakeUploadId, const FString& InMessage, int32 InCode)
+{
+	using namespace UE::CaptureManager;
+
+	if (!ensure(InTakeUploadId.IsValid()))
+	{
+		return;
+	}
+
+	FScopeLock Lock(&Mutex);
+	if (const FContext* Context = ImportContexts.Find(InTakeUploadId))
+	{
+		const bool bShouldCleanConversionDirectory = GetDefault<UCaptureManagerSettings>()->bShouldCleanConversionDirectory;
+		if (bShouldCleanConversionDirectory)
+		{
+			const FString ConversionDirectory = FPaths::Combine(Context->IngestOptions->Location,
+																Context->TakeName + TEXT("_") + InTakeUploadId.ToString());
+			IFileManager::Get().DeleteDirectory(*ConversionDirectory, false, true);
+		}
+
+		if (InMessage.IsEmpty())
+		{
+			ExecuteProcessFinishedReporter(Context->TakeId, MakeValue());
+		}
+		else
+		{
+			ExecuteProcessFinishedReporter(Context->TakeId, MakeError(FIngestCapability_Error::UploaderError, InMessage));
+		}
+
+		ImportContexts.Remove(InTakeUploadId);
+
+		WorkerForTake.Remove(Context->TakeId);
+	}
+}
+
+TOptional<FTakeMetadata> UGenericLiveLinkCPSDevice::ParseTake(const FString& InTakeName)
+{
+	using namespace UE::CaptureManager;
+
+	const FString DownloadedStorage = GetDefault<UCaptureManagerSettings>()->DownloadDirectory.Path;
+	const FString DownloadedTake = FPaths::Combine(DownloadedStorage, InTakeName);
+
+	TArray<FString> TakeFiles;
+	IFileManager::Get().FindFiles(TakeFiles, *DownloadedTake, *FTakeMetadata::FileExtension);
+	if (!TakeFiles.IsEmpty())
+	{
+		check(TakeFiles.Num() == 1);
+
+		FTakeMetadataParser TakeMetadataParser;
+		TValueOrError<FTakeMetadata, FTakeMetadataParserError> TakeMetadataResult = TakeMetadataParser.Parse(DownloadedTake / TakeFiles[0]);
+		if (!TakeMetadataResult.HasError())
+		{
+			return TakeMetadataResult.StealValue();
+		}
+	}
+
+	TArray<FText> ValidationFailures;
+	TOptional<FTakeMetadata> ParseOldMetadataResult = LiveLinkMetadata::ParseOldLiveLinkTakeMetadata(DownloadedTake, ValidationFailures);
+
+	if (ParseOldMetadataResult.IsSet())
+	{
+		return ParseOldMetadataResult.GetValue();
+	}
+
+	ParseOldMetadataResult = StereoCameraMetadata::ParseOldStereoCameraMetadata(DownloadedTake, ValidationFailures);
+
+	if (ParseOldMetadataResult.IsSet())
+	{
+		return ParseOldMetadataResult.GetValue();
+	}
+
+	return {};
+}
+
+void UGenericLiveLinkCPSDevice::FetchThumbnail(const FString& InTakeName, UE::CaptureManager::FTakeId InTakeId)
+{
+	TMap<FString, UE::CaptureManager::FTakeId> NameToIdMap;
+	NameToIdMap.Add(InTakeName, InTakeId);
+
+	FetchThumbnails(MoveTemp(NameToIdMap));
+}
+
+void UGenericLiveLinkCPSDevice::FetchThumbnails(TMap<FString, UE::CaptureManager::FTakeId> InNameToIdMap)
+{
+	using namespace UE::CaptureManager;
+
+	FCPSDataStream::FFileExportFinished OnFileExportFinishedCallback = FCPSDataStream::FFileExportFinished::CreateLambda(
+		[this, NameToIdMap = MoveTemp(InNameToIdMap)](const FString& InTakeName, TProtocolResult<FCPSDataStream::FData> InData)
+		{
+			if (InData.HasError())
+			{
+				return;
+			}
+
+			FCPSDataStream::FData Data = InData.StealValue();
+
+			if (const FTakeId* TakeId = NameToIdMap.Find(InTakeName))
+			{
+				TOptional<FTakeMetadata> TakeMetadataOpt = GetTakeMetadata(*TakeId);
+				if (!TakeMetadataOpt.IsSet())
+				{
+					return;
+				}
+
+				FTakeMetadata TakeMetadata = TakeMetadataOpt.GetValue();
+				TakeMetadata.Thumbnail = FTakeThumbnailData(MoveTemp(Data));
+
+				UpdateTake(*TakeId, MoveTemp(TakeMetadata));
+
+				PublishEvent<FIngestCapability_TakeUpdatedEvent>(*TakeId);
+			}
+		});
+
+	TUniquePtr<FCPSDataStream> DataStream = MakeUnique<FCPSDataStream>(MoveTemp(OnFileExportFinishedCallback));
+
+	Device->FetchThumbnails(MoveTemp(DataStream));
+}
diff --git a/Engine/Plugins/VirtualProduction/CaptureManager/CaptureManagerApp/Source/GenericLiveLinkDevice/Private/CPS/GenericLiveLinkCPSDevice.h b/Engine/Plugins/VirtualProduction/CaptureManager/CaptureManagerApp/Source/GenericLiveLinkDevice/Private/CPS/GenericLiveLinkCPSDevice.h
new file mode 100644
--- /dev/null
+++ b/Engine/Plugins/VirtualProduction/CaptureManager/CaptureManagerApp/Source/GenericLiveLinkDevice/Private/CPS/GenericLiveLinkCPSDevice.h
@@ -0,0 +1,147 @@
+// Copyright Epic Games, Inc. All Rights Reserved.
+
+#pragma once
+
+#include "LiveLinkDevice.h"
+
+#include "Ingest/LiveLinkDeviceCapability_Ingest.h"
+#include "LiveLinkDeviceCapability_Connection.h"
+#include "LiveLinkDeviceCapability_Recording.h"
+#include "Misc/Optional.h"
+
+#include "Async/TaskProgress.h"
+#include "LiveLinkHubExportManager.h"
+
+#include "CaptureDataConverter.h"
+
+#include "CaptureCommunicationManager.h"
+
+#include "Protocol/CPSDevice.h"
+#include "Customizations/ToggleConnectActionCustomization.h"
+#include "Customizations/DeviceIpAddressCustomization.h"
+
+#include "GenericLiveLinkCPSDevice.generated.h"
+
+
+UCLASS()
+class UGenericLiveLinkCPSDeviceSettings : public ULiveLinkDeviceSettings
+{
+	GENERATED_BODY()
+
+public:
+	UPROPERTY(EditAnywhere, Category = "Live Link CPS Device")
+	FString DisplayName = TEXT("Live Link CPS Device");
+
+	UPROPERTY(EditAnywhere, Category = "Live Link CPS Device")
+	FDeviceIpAddress IpAddress;
+
+	UPROPERTY(EditAnywhere, Category = "Live Link CPS Device")
+	uint16 Port = 14785;
+
+	UPROPERTY(VisibleAnywhere, Category = "Live Link CPS Device")
+	FToggleConnectAction ConnectAction;
+};
+
+
+UCLASS()
+class UGenericLiveLinkCPSDevice : public ULiveLinkDevice
+	, public ILiveLinkDeviceCapability_Ingest
+	, public ILiveLinkDeviceCapability_Connection
+	, public ILiveLinkDeviceCapability_Recording
+{
+	GENERATED_BODY()
+
+public:
+	UGenericLiveLinkCPSDeviceSettings* GetSettings();
+	const UGenericLiveLinkCPSDeviceSettings* GetSettings() const;
+
+	//~ Begin ULiveLinkDevice interface
+	virtual TSubclassOf<ULiveLinkDeviceSettings> StaticGetSettingsClass() const override;
+	virtual FText GetDisplayName() const override;
+	virtual EDeviceHealth GetDeviceHealth() const override;
+	virtual FText GetHealthText() const override;
+
+	virtual void OnDeviceAdded(const FGuid InDeviceGuid, ULiveLinkDeviceSettings* InSettings) override;
+	virtual void OnDeviceRemoved() override;
+	virtual void OnSettingChanged(const FPropertyChangedEvent& InPropertyChangedEvent) override;
+	//~ End ULiveLinkDevice interface
+
+private:
+	//~ Begin ULiveLinkDeviceCapability_Ingest interface
+	virtual void UpdateTakeListAsync_Implementation() override;
+	virtual void RunIngestTakeAsync(int32 InTakeId,
+									const UIngestCapability_Options* InIngestOptions) override;
+	virtual void CancelIngest(int32 InTakeId) override;
+	//~ End ULiveLinkDeviceCapability_Ingest interface
+
+	//~ Begin ILiveLinkDeviceCapability_Connection interface
+	virtual ELiveLinkDeviceConnectionStatus GetConnectionStatus_Implementation() const override;
+	virtual FString GetHardwareId_Implementation() const override;
+	virtual bool SetHardwareId_Implementation(const FString& HardwareID) override;
+	virtual bool Connect_Implementation() override;
+	virtual bool Disconnect_Implementation() override;
+	//~ End ILiveLinkDeviceCapability_Connection interface
+
+	//~ Begin ILiveLinkDeviceCapability_Recording interface
+	virtual bool StartRecording_Implementation() override;
+	virtual bool StopRecording_Implementation() override;
+	virtual bool IsRecording_Implementation() const override;
+	//~ End ILiveLinkDeviceCapability_Recording interface
+
+	void HandleConnectionChanged(TSharedPtr<const UE::CaptureManager::FCaptureEvent> InEvent);
+	void HandleCPSStateUpdate(TSharedPtr<const UE::CaptureManager::FCaptureEvent> InEvent);
+	void HandleCPSEvent(TSharedPtr<const UE::CaptureManager::FCaptureEvent> InEvent);
+
+	void OnExportProgressReport(float InProgress,
+								UE::CaptureManager::FTaskProgress::FTask InTask);
+	void OnExportFinished(UE::CaptureManager::TProtocolResult<void> InResult,
+						  FString InTakeName,
+						  UE::CaptureManager::FTakeId InTakeId,
+						  TSharedPtr<UE::CaptureManager::FTaskProgress> InTaskProgress,
+						  TStrongObjectPtr<const UIngestCapability_Options> InIngestOptions);
+
+	void ConvertData(const FString& InTakeName,
+					 UE::CaptureManager::FTakeId InTakeId,
+					 TSharedPtr<UE::CaptureManager::FTaskProgress> InTaskProgress,
+					 TStrongObjectPtr<const UIngestCapability_Options> InIngestOptions);
+
+	FCaptureDataConverterResult<void> RunConversion(int32 InTakeId,
+													TStrongObjectPtr<const UIngestCapability_Options> InIngestOptions,
+													TSharedPtr<UE::CaptureManager::FTaskProgress> InTaskProgress,
+													FCaptureDataConverterParams& InParams);
+
+	void UploadTakeStateHandler(const FGuid& InTakeUploadId, double InProgress);
+	void UploadTakeFinished(const FGuid& InTakeUploadId, const FString& InMessage, int32 InCode);
+
+	TOptional<FTakeMetadata> ParseTake(const FString& InTakeName);
+	FTakeMetadata ParseTakeMetadata(const UE::CaptureManager::FGetTakeMetadataResponse::FTakeObject& InTake);
+
+	void FetchThumbnail(const FString& InTakeName, UE::CaptureManager::FTakeId InTakeId);
+	void FetchThumbnails(TMap<FString, UE::CaptureManager::FTakeId> InNameToIdMap);
+
+	struct FContext
+	{
+		FString TakeName;
+		int32 TakeId;
+		TStrongObjectPtr<const UIngestCapability_Options> IngestOptions;
+		TSharedPtr<UE::CaptureManager::FTaskProgress> TaskProgress;
+		UE::CaptureManager::FTaskProgress::FTask Task;
+	};
+
+	TSharedPtr<UE::CaptureManager::FCPSDevice> Device;
+
+	TSharedPtr<FCommunicationManager> CommunicationManager;
+	TWeakPtr<FUploader> Uploader;
+
+	TMap<FGuid, FContext> ImportContexts;
+
+	FCriticalSection Mutex;
+	TMap<int32, int32> WorkerForTake;
+
+	TMap<int32, TUniquePtr<FCaptureDataConverter>> Converters;
+
+	FGuid DeviceGuid;
+
+	std::atomic_bool bIsRecording = false;
+	std::atomic_bool bIsConnecting = false;
+};
diff --git a/Engine/Plugins/VirtualProduction/CaptureManager/CaptureManagerApp/Source/GenericLiveLinkDevice/Private/GenericLiveLinkDeviceModule.cpp b/Engine/Plugins/VirtualProduction/CaptureManager/CaptureManagerApp/Source/GenericLiveLinkDevice/Private/GenericLiveLinkDeviceModule.cpp
new file mode 100644
--- /dev/null
+++ b/Engine/Plugins/VirtualProduction/CaptureManager/CaptureManagerApp/Source/GenericLiveLinkDevice/Private/GenericLiveLinkDeviceModule.cpp
@@ -0,0 +1,43 @@
+// Copyright Epic Games, Inc. All Rights Reserved.
+
+#include "Modules/ModuleManager.h"
+
+#include "PropertyEditorModule.h"
+
+#include "CPS/GenericLiveLinkCPSDevice.h"
+#include "CPS/Customizations/ToggleConnectActionCustomization.h"
+#include "CPS/Customizations/DeviceIpAddressCustomization.h"
+
+class FGenericLiveLinkDeviceModule : public IModuleInterface
+{
+public:
+
+	void StartupModule()
+	{
+		FPropertyEditorModule& PropertyModule = FModuleManager::LoadModuleChecked<FPropertyEditorModule>("PropertyEditor");
+		PropertyModule.RegisterCustomPropertyTypeLayout(FToggleConnectAction::StaticStruct()->GetFName(),
+														FOnGetPropertyTypeCustomizationInstance::CreateRaw(this, &FGenericLiveLinkDeviceModule::GetToggleConnectCustomization));
+		PropertyModule.RegisterCustomPropertyTypeLayout(FDeviceIpAddress::StaticStruct()->GetFName(),
+														FOnGetPropertyTypeCustomizationInstance::CreateRaw(this, &FGenericLiveLinkDeviceModule::GetDeviceIpAddressCustomization));
+	}
+
+	void ShutdownModule()
+	{
+		FPropertyEditorModule& PropertyModule = FModuleManager::LoadModuleChecked<FPropertyEditorModule>("PropertyEditor");
+		PropertyModule.UnregisterCustomPropertyTypeLayout(FToggleConnectAction::StaticStruct()->GetFName());
+	}
+
+private:
+
+	TSharedRef<IPropertyTypeCustomization> GetToggleConnectCustomization()
+	{
+		return MakeShared<FToggleConnectActionCustomization>();
+	}
+
+	TSharedRef<IPropertyTypeCustomization> GetDeviceIpAddressCustomization()
+	{
+		return MakeShared<FDeviceIpAddressCustomization>();
+	}
+};
+
+IMPLEMENT_MODULE(FGenericLiveLinkDeviceModule, GenericLiveLinkDevice);
\ No newline at end of file
diff --git a/Engine/Plugins/VirtualProduction/CaptureManager/CaptureManagerApp/Source/GenericLiveLinkDevice/Private/Tests/GenericLiveLinkDeviceTests.cpp b/Engine/Plugins/VirtualProduction/CaptureManager/CaptureManagerApp/Source/GenericLiveLinkDevice/Private/Tests/GenericLiveLinkDeviceTests.cpp
new file mode 100644
--- /dev/null
+++ b/Engine/Plugins/VirtualProduction/CaptureManager/CaptureManagerApp/Source/GenericLiveLinkDevice/Private/Tests/GenericLiveLinkDeviceTests.cpp
@@ -0,0 +1,38 @@
+// Copyright Epic Games, Inc. All Rights Reserved.
+
+#include "Archive/GenericLiveLinkArchiveDevice.h"
+
+#include "Misc/AutomationTest.h"
+
+IMPLEMENT_SIMPLE_AUTOMATION_TEST(FGenericLiveLinkArchiveDeviceCreationTest, "CaptureManager.GenericLiveLinkArchiveDevice.Creation",
+								 EAutomationTestFlags::EditorContext | EAutomationTestFlags::EngineFilter);
+
+bool FGenericLiveLinkArchiveDeviceCreationTest::RunTest(const FString& Parameters)
+{
+	UGenericLiveLinkArchiveDevice* Device = NewObject<UGenericLiveLinkArchiveDevice>();
+
+	TestNotNull(TEXT("Device Not Null"), Device);
+	TestTrue(TEXT("Device Implements Ingest Capability"), Device->Implements<ULiveLinkDeviceCapability_Ingest>());
+
+	return true;
+}
+
+IMPLEMENT_SIMPLE_AUTOMATION_TEST(FGenericLiveLinkArchiveDeviceIngestTakeAsyncTest, "CaptureManager.GenericLiveLinkArchiveDevice.IngestTakeAsync",
+								 EAutomationTestFlags::EditorContext | EAutomationTestFlags::EngineFilter);
+
+bool FGenericLiveLinkArchiveDeviceIngestTakeAsyncTest::RunTest(const FString& Parameters)
+{
+	UGenericLiveLinkArchiveDevice* Device = NewObject<UGenericLiveLinkArchiveDevice>();
+
+	TestNotNull(TEXT("Device Not Null"), Device);
+	TestTrue(TEXT("Device Implements Ingest Capability"), Device->Implements<ULiveLinkDeviceCapability_Ingest>());
+
+	UIngestCapability_ProcessHandle* IngestProcessHandle =
+		ILiveLinkDeviceCapability_Ingest::Execute_PrepareIngestProcess(Device, 0);
+	TestNotNull(TEXT("Ingest Process Handle Not Null"), IngestProcessHandle);
+
+	static constexpr UIngestCapability_Options* Options = nullptr;
+	IngestProcessHandle->Run(Options);
+
+	return true;
+}
\ No newline at end of file
diff --git a/Engine/Plugins/VirtualProduction/CaptureManager/CaptureManagerApp/Source/GenericLiveLinkDevice/Private/Utils/GenericLiveLinkDeviceUtility.cpp b/Engine/Plugins/VirtualProduction/CaptureManager/CaptureManagerApp/Source/GenericLiveLinkDevice/Private/Utils/GenericLiveLinkDeviceUtility.cpp
new file mode 100644
--- /dev/null
+++ b/Engine/Plugins/VirtualProduction/CaptureManager/CaptureManagerApp/Source/GenericLiveLinkDevice/Private/Utils/GenericLiveLinkDeviceUtility.cpp
@@ -0,0 +1,60 @@
+// Copyright Epic Games, Inc. All Rights Reserved.
+
+#include "GenericLiveLinkDeviceUtility.h"
+
+namespace UE::CaptureManager::Private
+{
+
+TSharedRef<FCommunicationManager> GetCommunicationManager()
+{
+	FCaptureCommunicationManagerModule& CaptureCommunicationManagerModule =
+		FModuleManager::LoadModuleChecked<FCaptureCommunicationManagerModule>("CaptureCommunicationManager");
+	return CaptureCommunicationManagerModule.GetManager();
+}
+
+EMediaTexturePixelFormat ConvertPixelFormat(EIngestCapability_ImagePixelFormat InImagePixelFormat)
+{
+	switch (InImagePixelFormat)
+	{
+		case EIngestCapability_ImagePixelFormat::U8_RGB:
+			return EMediaTexturePixelFormat::U8_RGB;
+		case EIngestCapability_ImagePixelFormat::U8_BGR:
+			return EMediaTexturePixelFormat::U8_BGR;
+		case EIngestCapability_ImagePixelFormat::U8_RGBA:
+			return EMediaTexturePixelFormat::U8_RGBA;
+		case EIngestCapability_ImagePixelFormat::U8_BGRA:
+			return EMediaTexturePixelFormat::U8_BGRA;
+		case EIngestCapability_ImagePixelFormat::U8_I444:
+			return EMediaTexturePixelFormat::U8_I444;
+		case EIngestCapability_ImagePixelFormat::U8_I420:
+			return EMediaTexturePixelFormat::U8_I420;
+		case EIngestCapability_ImagePixelFormat::U8_NV12:
+			return EMediaTexturePixelFormat::U8_NV12;
+		case EIngestCapability_ImagePixelFormat::U8_Mono:
+			return EMediaTexturePixelFormat::U8_Mono;
+		case EIngestCapability_ImagePixelFormat::U16_Mono:
+			return EMediaTexturePixelFormat::U16_Mono;
+		case EIngestCapability_ImagePixelFormat::F_Mono:
+			return EMediaTexturePixelFormat::F_Mono;
+		case EIngestCapability_ImagePixelFormat::Undefined:
+		default:
+			check(false);
+			return EMediaTexturePixelFormat::Undefined;
+	}
+}
+
+FString ErrorOriginToString(FTakeMetadataParserError::EOrigin InOrigin)
+{
+	switch (InOrigin)
+	{
+		case FTakeMetadataParserError::Reader:
+			return TEXT("Reader");
+		case FTakeMetadataParserError::Validator:
+			return TEXT("Validator");
+		case FTakeMetadataParserError::Parser:
+		default:
+			return TEXT("Parser");
+	}
+}
+
+}
\ No newline at end of file
diff --git a/Engine/Plugins/VirtualProduction/CaptureManager/CaptureManagerApp/Source/GenericLiveLinkDevice/Private/Utils/GenericLiveLinkDeviceUtility.h b/Engine/Plugins/VirtualProduction/CaptureManager/CaptureManagerApp/Source/GenericLiveLinkDevice/Private/Utils/GenericLiveLinkDeviceUtility.h
new file mode 100644
--- /dev/null
+++ b/Engine/Plugins/VirtualProduction/CaptureManager/CaptureManagerApp/Source/GenericLiveLinkDevice/Private/Utils/GenericLiveLinkDeviceUtility.h
@@ -0,0 +1,17 @@
+// Copyright Epic Games, Inc. All Rights Reserved.
+
+#pragma once
+
+#include "CaptureCommunicationManagerModule.h"
+#include "CaptureManagerTakeMetadata.h"
+#include "Ingest/IngestCapability_Options.h"
+#include "MediaSample.h"
+
+namespace UE::CaptureManager::Private
+{
+
+TSharedRef<FCommunicationManager> GetCommunicationManager();
+EMediaTexturePixelFormat ConvertPixelFormat(EIngestCapability_ImagePixelFormat InImagePixelFormat);
+FString ErrorOriginToString(FTakeMetadataParserError::EOrigin InOrigin);
+
+}
\ No newline at end of file
diff --git a/Engine/Plugins/VirtualProduction/CaptureManager/CaptureManagerApp/Source/LiveLinkCapabilities/Private/Ingest/IngestCapability_ProcessHandle.cpp b/Engine/Plugins/VirtualProduction/CaptureManager/CaptureManagerApp/Source/LiveLinkCapabilities/Private/Ingest/IngestCapability_ProcessHandle.cpp
new file mode 100644
--- /dev/null
+++ b/Engine/Plugins/VirtualProduction/CaptureManager/CaptureManagerApp/Source/LiveLinkCapabilities/Private/Ingest/IngestCapability_ProcessHandle.cpp
@@ -0,0 +1,67 @@
+// Copyright Epic Games, Inc. All Rights Reserved.
+
+#include "Ingest/IngestCapability_ProcessHandle.h"
+
+#include "Ingest/LiveLinkDeviceCapability_Ingest.h"
+
+FIngestCapability_Error::FIngestCapability_Error(ECode InCode, FString InMessage)
+	: Code(InCode)
+	, Message(MoveTemp(InMessage))
+{
+}
+
+FIngestCapability_Error::ECode FIngestCapability_Error::GetCode() const
+{
+	return Code;
+}
+
+const FString& FIngestCapability_Error::GetMessage() const
+{
+	return Message;
+}
+
+UIngestCapability_ProcessHandle::UIngestCapability_ProcessHandle()
+	: TakeId(INDEX_NONE)
+	, Owner(nullptr)
+{
+}
+
+void UIngestCapability_ProcessHandle::Initialize(int32 InTakeId, ILiveLinkDeviceCapability_Ingest* InOwner)
+{
+	TakeId = InTakeId;
+	Owner = InOwner;
+}
+
+void UIngestCapability_ProcessHandle::Run(const UIngestCapability_Options* InIngestOptions)
+{
+	if (ensure(InIngestOptions))
+	{
+		Owner->RunIngestTakeAsync(TakeId, InIngestOptions);
+	}
+}
+
+void UIngestCapability_ProcessHandle::Cancel()
+{
+	Owner->CancelIngest(TakeId);
+}
+
+bool UIngestCapability_ProcessResult::IsValid() const
+{
+	return Code == 0;
+}
+
+bool UIngestCapability_ProcessResult::IsError() const
+{
+	return !IsValid();
+}
+
+UIngestCapability_ProcessResult* UIngestCapability_ProcessResult::Create(int32 InCode, FText InMessage)
+{
+	UIngestCapability_ProcessResult* Result =
+		NewObject<UIngestCapability_ProcessResult>();
+
+	Result->Code = InCode;
+	Result->Message = InMessage;
+
+	return Result;
+}
\ No newline at end of file
diff --git a/Engine/Plugins/VirtualProduction/CaptureManager/CaptureManagerApp/Source/LiveLinkCapabilities/Private/Ingest/LiveLinkDeviceCapability_Ingest.cpp b/Engine/Plugins/VirtualProduction/CaptureManager/CaptureManagerApp/Source/LiveLinkCapabilities/Private/Ingest/LiveLinkDeviceCapability_Ingest.cpp
new file mode 100644
--- /dev/null
+++ b/Engine/Plugins/VirtualProduction/CaptureManager/CaptureManagerApp/Source/LiveLinkCapabilities/Private/Ingest/LiveLinkDeviceCapability_Ingest.cpp
@@ -0,0 +1,209 @@
+// Copyright Epic Games, Inc. All Rights Reserved.
+
+#include "Ingest/LiveLinkDeviceCapability_Ingest.h"
+
+#include "LiveLinkDevice.h"
+#include "Widgets/Text/STextBlock.h"
+
+SHeaderRow::FColumn::FArguments& ULiveLinkDeviceCapability_Ingest::GenerateHeaderForColumn(const FName InColumnId, SHeaderRow::FColumn::FArguments& InArgs) const
+{
+	if (InColumnId == Column_IngestSupport)
+	{
+		return InArgs.DefaultLabel(FText::FromString("Ingest"))
+			.DefaultTooltip(FText::FromString("Device supports ingest"))
+			.FillSized(30.0f);
+	}
+
+	return Super::GenerateHeaderForColumn(InColumnId, InArgs);
+}
+
+TSharedPtr<SWidget> ULiveLinkDeviceCapability_Ingest::GenerateWidgetForColumn(const FName InColumnId, ULiveLinkDevice* InDevice) const
+{
+	return SNew(STextBlock)
+		.Text_Lambda(
+			[WeakDevice = TWeakObjectPtr<ULiveLinkDevice>(InDevice)]
+			()
+			{
+				if (ULiveLinkDevice* Device = WeakDevice.Get())
+				{
+					const bool bDeviceSupportsIngest = Device->Implements<ULiveLinkDeviceCapability_Ingest>();
+					if (bDeviceSupportsIngest)
+					{
+						return FText::FromString(TEXT("Y"));
+					}
+				}
+
+				return FText::FromString(TEXT("N"));
+			}
+		);
+}
+
+ILiveLinkDeviceCapability_Ingest::ILiveLinkDeviceCapability_Ingest()
+{
+	RegisterEvent(FString(FIngestCapability_TakeAddedEvent::Name));
+	RegisterEvent(FString(FIngestCapability_TakeUpdatedEvent::Name));
+	RegisterEvent(FString(FIngestCapability_TakeRemovedEvent::Name));
+}
+
+UIngestCapability_ProcessHandle* ILiveLinkDeviceCapability_Ingest::PrepareIngestProcess_Implementation(int32 InTakeId)
+{
+	TStrongObjectPtr<UIngestCapability_ProcessHandle> IngestProcess = CreateIngestProcessHandle(InTakeId, this);
+
+	UIngestCapability_ProcessHandle* ReturnValue = IngestProcess.Get();
+
+	FScopeLock Lock(&IngestProcessMutex);
+	IngestProcessHandles.Add(InTakeId, MoveTemp(IngestProcess));
+
+	return ReturnValue;
+}
+
+UIngestCapability_TakeInformation* ILiveLinkDeviceCapability_Ingest::GetTakeInformation_Implementation(int32 InTakeId) const
+{
+	TOptional<FTakeMetadata> TakeMetadataOpt = GetTakeMetadata(InTakeId);
+
+	if (!TakeMetadataOpt.IsSet())
+	{
+		return nullptr;
+	}
+
+	FTakeMetadata TakeMetadata = TakeMetadataOpt.GetValue();
+
+	UIngestCapability_TakeInformation* TakeInformation = NewObject<UIngestCapability_TakeInformation>();
+	TakeInformation->SlateName = TakeMetadata.Slate;
+	TakeInformation->TakeNumber = TakeMetadata.TakeNumber;
+
+	if (TakeMetadata.DateTime.IsSet())
+	{
+		TakeInformation->DateTime = TakeMetadata.DateTime.GetValue();
+	}
+
+	return TakeInformation;
+}
+
+TArray<int32> ILiveLinkDeviceCapability_Ingest::GetTakeIdentifiers_Implementation() const
+{
+	TArray<int32> TakeIdentifiers;
+
+	FScopeLock Lock(&TakeMetadataMutex);
+	TakeMetadataMap.GetKeys(TakeIdentifiers);
+
+	return TakeIdentifiers;
+}
+
+UIngestCapability_UpdateTakeListCallback* ILiveLinkDeviceCapability_Ingest::CreateUpdateTakeListCallback_Implementation()
+{
+	UpdateTakeListCallback = TStrongObjectPtr(NewObject<UIngestCapability_UpdateTakeListCallback>());
+
+	return UpdateTakeListCallback.Get();
+}
+
+void ILiveLinkDeviceCapability_Ingest::ExecuteProcessFinishedReporter(int32 InTakeIdentifier,
+																	  TValueOrError<void, FIngestCapability_Error> InMaybeError)
+{
+	using namespace UE::CaptureManager;
+
+	FScopeLock Lock(&IngestProcessMutex);
+	TStrongObjectPtr<UIngestCapability_ProcessHandle>* IngestProcessHandlePtr = IngestProcessHandles.Find(InTakeIdentifier);
+	if (!IngestProcessHandlePtr)
+	{
+		return;
+	}
+
+	TStrongObjectPtr<UIngestCapability_ProcessHandle> IngestProcessHandle = *IngestProcessHandlePtr;
+	IngestProcessHandles.Remove(InTakeIdentifier);
+	Lock.Unlock();
+
+	UIngestCapability_ProcessResult* Result =
+		UIngestCapability_ProcessResult::Create();
+
+	if (InMaybeError.HasError())
+	{
+		FIngestCapability_Error Error = InMaybeError.GetError();
+		Result->Code = static_cast<int32>(Error.GetCode());
+		Result->Message = FText::FromString(Error.GetMessage());
+	}
+
+	IngestProcessHandle->ProcessFinishedReporterDynamic.ExecuteIfBound(InTakeIdentifier, Result);
+	IngestProcessHandle->ProcessFinishedReporter(InTakeIdentifier, MoveTemp(InMaybeError));	
+}
+
+void ILiveLinkDeviceCapability_Ingest::ExecuteProcessProgressReporter(int32 InTakeIdentifier, double InProgress)
+{
+	FScopeLock Lock(&IngestProcessMutex);
+	TStrongObjectPtr<UIngestCapability_ProcessHandle>* IngestProcessHandlePtr = IngestProcessHandles.Find(InTakeIdentifier);
+	if (!IngestProcessHandlePtr)
+	{
+		return;
+	}
+
+	TStrongObjectPtr<UIngestCapability_ProcessHandle> IngestProcessHandle = *IngestProcessHandlePtr;
+	Lock.Unlock();
+
+	IngestProcessHandle->ProcessProgressReporterDynamic.ExecuteIfBound(InTakeIdentifier, InProgress);
+	IngestProcessHandle->ProcessProgressReporter(InTakeIdentifier, InProgress);
+}
+
+void ILiveLinkDeviceCapability_Ingest::ExecuteUpdateTakeListCallback(const TArray<int32>& InTakeIdentifiers)
+{
+	if (ensure(UpdateTakeListCallback))
+	{
+		UpdateTakeListCallback->DynamicCallback.ExecuteIfBound(InTakeIdentifiers);
+		UpdateTakeListCallback->Callback(InTakeIdentifiers);
+	}
+
+	UpdateTakeListCallback = nullptr;
+}
+
+int32 ILiveLinkDeviceCapability_Ingest::AddTake(FTakeMetadata InTakeMetadata)
+{
+	int32 NewTakeId = CurrentTakeId.fetch_add(1);
+
+	FScopeLock Lock(&TakeMetadataMutex);
+	TakeMetadataMap.Add(NewTakeId, MoveTemp(InTakeMetadata));
+
+	return NewTakeId;
+}
+
+void ILiveLinkDeviceCapability_Ingest::RemoveTake(int32 InTakeId)
+{
+	FScopeLock Lock(&TakeMetadataMutex);
+	TakeMetadataMap.Remove(InTakeId);
+}
+
+void ILiveLinkDeviceCapability_Ingest::RemoveAllTakes()
+{
+	FScopeLock Lock(&TakeMetadataMutex);
+	TakeMetadataMap.Empty();
+}
+
+bool ILiveLinkDeviceCapability_Ingest::UpdateTake(int32 InTakeId, FTakeMetadata InTakeMetadata)
+{
+	FScopeLock Lock(&TakeMetadataMutex);
+	if (FTakeMetadata* TakeMetadata = TakeMetadataMap.Find(InTakeId))
+	{
+		*TakeMetadata = MoveTemp(InTakeMetadata);
+		return true;
+	}
+
+	return false;
+}
+
+TOptional<FTakeMetadata> ILiveLinkDeviceCapability_Ingest::GetTakeMetadata(int32 InTakeId) const
+{
+	FScopeLock Lock(&TakeMetadataMutex);
+	if (const FTakeMetadata* TakeMetadata = TakeMetadataMap.Find(InTakeId))
+	{
+		return *TakeMetadata;
+	}
+
+	return {};
+}
+
+TStrongObjectPtr<UIngestCapability_ProcessHandle> 
+ILiveLinkDeviceCapability_Ingest::CreateIngestProcessHandle(int32 InTakeId, ILiveLinkDeviceCapability_Ingest* InOwner)
+{
+	TStrongObjectPtr<UIngestCapability_ProcessHandle> IngestProcessHandle(NewObject<UIngestCapability_ProcessHandle>());
+	IngestProcessHandle->Initialize(InTakeId, InOwner);
+
+	return IngestProcessHandle;
+}
\ No newline at end of file
diff --git a/Engine/Plugins/VirtualProduction/CaptureManager/CaptureManagerApp/Source/LiveLinkCapabilities/Public/Ingest/IngestCapability_Options.h b/Engine/Plugins/VirtualProduction/CaptureManager/CaptureManagerApp/Source/LiveLinkCapabilities/Public/Ingest/IngestCapability_Options.h
new file mode 100644
--- /dev/null
+++ b/Engine/Plugins/VirtualProduction/CaptureManager/CaptureManagerApp/Source/LiveLinkCapabilities/Public/Ingest/IngestCapability_Options.h
@@ -0,0 +1,69 @@
+// Copyright Epic Games, Inc. All Rights Reserved.
+
+#pragma once
+
+#include "LiveLinkDevice.h"
+
+#include "IngestCapability_Options.generated.h"
+
+UENUM()
+enum class EIngestCapability_ImagePixelFormat
+{
+	Undefined = -1		UMETA(Hidden),
+	U8_RGB = 0			UMETA(DisplayName = "RGB 8-bit"),
+	U8_BGR				UMETA(DisplayName = "BGR 8-bit"),
+	U8_RGBA				UMETA(DisplayName = "RGBA 8-bit"),
+	U8_BGRA				UMETA(DisplayName = "BGRA 8-bit"),
+	U8_I444				UMETA(DisplayName = "I444 8-bit"),
+	U8_I420				UMETA(DisplayName = "I420 8-bit"),
+	U8_NV12				UMETA(DisplayName = "NV12 8-bit"),
+	U8_Mono				UMETA(DisplayName = "Mono 8-bit"),
+	U16_Mono			UMETA(DisplayName = "Mono 16-bit"),
+	F_Mono				UMETA(DisplayName = "Mono Float"),
+
+	Default = U8_RGB	UMETA(DisplayName = "Default (U8_RGB)")
+};
+
+USTRUCT()
+struct LIVELINKCAPABILITIES_API FIngestCapability_VideoOptions
+{
+	GENERATED_USTRUCT_BODY()
+
+	UPROPERTY(EditAnywhere, Category = "Live Link Device|Ingest")
+	FString FileNamePrefix;
+
+	UPROPERTY(EditAnywhere, Category = "Live Link Device|Ingest")
+	FString Format;
+
+	UPROPERTY(EditAnywhere, Category = "Live Link Device|Ingest")
+	EIngestCapability_ImagePixelFormat PixelFormat;
+};
+
+USTRUCT()
+struct LIVELINKCAPABILITIES_API FIngestCapability_AudioOptions
+{
+	GENERATED_USTRUCT_BODY()
+
+	UPROPERTY(EditAnywhere, Category = "Live Link Device|Ingest")
+	FString FileNamePrefix;
+
+	UPROPERTY(EditAnywhere, Category = "Live Link Device|Ingest")
+	FString Format;
+};
+
+UCLASS()
+class LIVELINKCAPABILITIES_API UIngestCapability_Options : public UObject
+{
+	GENERATED_BODY()
+
+public:
+
+	UPROPERTY(EditAnywhere, Category = "Live Link Device|Ingest")
+	FString Location;
+
+	UPROPERTY(EditAnywhere, Category = "Live Link Device|Ingest")
+	FIngestCapability_VideoOptions Video;
+
+	UPROPERTY(EditAnywhere, Category = "Live Link Device|Ingest")
+	FIngestCapability_AudioOptions Audio;
+};
\ No newline at end of file
diff --git a/Engine/Plugins/VirtualProduction/CaptureManager/CaptureManagerApp/Source/LiveLinkCapabilities/Public/Ingest/IngestCapability_ProcessHandle.h b/Engine/Plugins/VirtualProduction/CaptureManager/CaptureManagerApp/Source/LiveLinkCapabilities/Public/Ingest/IngestCapability_ProcessHandle.h
new file mode 100644
--- /dev/null
+++ b/Engine/Plugins/VirtualProduction/CaptureManager/CaptureManagerApp/Source/LiveLinkCapabilities/Public/Ingest/IngestCapability_ProcessHandle.h
@@ -0,0 +1,103 @@
+// Copyright Epic Games, Inc. All Rights Reserved.
+
+#pragma once
+
+#include "LiveLinkDevice.h"
+
+#include "Templates/ValueOrError.h"
+#include "Async/ManagedDelegate.h"
+
+#include "IngestCapability_Options.h"
+#include "IngestCapability_TakeInformation.h"
+
+#include "IngestCapability_ProcessHandle.generated.h"
+
+class LIVELINKCAPABILITIES_API FIngestCapability_Error
+{
+public:
+
+	enum ECode
+	{
+		Ok = 0,
+		AbortedByUser = 1,
+		InternalError,
+		InvalidArgument,
+		DownloaderError,
+		UploaderNotFound,
+		UploaderError,
+		ConversionError
+	};
+
+	FIngestCapability_Error(ECode InCode, FString InMessage);
+
+	ECode GetCode() const;
+	const FString& GetMessage() const;
+
+private:
+
+	ECode Code;
+	FString Message;
+};
+
+UDELEGATE()
+DECLARE_DYNAMIC_DELEGATE_TwoParams(FProcessFinishReporter, int32, TakeIdentifier, UIngestCapability_ProcessResult*, IngestProcessResult);
+DECLARE_DYNAMIC_DELEGATE_TwoParams(FProcessProgressReporter, int32, TakeIdentifier, double, Progress);
+
+using FIngestProcessProgressReporter = UE::CaptureManager::TManagedDelegate<int32, double>;
+using FIngestProcessFinishReporter = UE::CaptureManager::TManagedDelegate<int32, TValueOrError<void, FIngestCapability_Error>>;
+
+UCLASS()
+class LIVELINKCAPABILITIES_API UIngestCapability_ProcessResult : public UObject
+{
+	GENERATED_BODY()
+
+public:
+
+	UPROPERTY()
+	FText Message;
+
+	UPROPERTY()
+	int32 Code = 0;
+
+	bool IsValid() const;
+	bool IsError() const;
+
+	static UIngestCapability_ProcessResult* Create(int32 InCode = 0, FText InMessage = FText());
+};
+
+UCLASS()
+class LIVELINKCAPABILITIES_API UIngestCapability_ProcessHandle : public UObject
+{
+	GENERATED_BODY()
+
+public:
+
+	UPROPERTY(EditAnywhere, Category = "Live Link Device|Ingest")
+	int32 TakeId = INDEX_NONE;
+
+	UFUNCTION(BlueprintCallable, Category = "Live Link Device|Ingest")
+	void Run(const UIngestCapability_Options* InIngestOptions);
+
+	UFUNCTION(BlueprintCallable, Category = "Live Link Device|Ingest")
+	void Cancel();
+
+	UPROPERTY(EditAnywhere, Category = "Live Link Device|Ingest")
+	FProcessFinishReporter ProcessFinishedReporterDynamic;
+
+	FIngestProcessFinishReporter ProcessFinishedReporter;
+
+	UPROPERTY(EditAnywhere, Category = "Live Link Device|Ingest")
+	FProcessProgressReporter ProcessProgressReporterDynamic;
+
+	FIngestProcessProgressReporter ProcessProgressReporter;
+
+private:
+
+	UIngestCapability_ProcessHandle();
+
+	void Initialize(int32 InTakeId, class ILiveLinkDeviceCapability_Ingest* InOwner);
+
+	class ILiveLinkDeviceCapability_Ingest* Owner;
+
+	friend class ILiveLinkDeviceCapability_Ingest;
+};
\ No newline at end of file
diff --git a/Engine/Plugins/VirtualProduction/CaptureManager/CaptureManagerApp/Source/LiveLinkCapabilities/Public/Ingest/IngestCapability_TakeInformation.h b/Engine/Plugins/VirtualProduction/CaptureManager/CaptureManagerApp/Source/LiveLinkCapabilities/Public/Ingest/IngestCapability_TakeInformation.h
new file mode 100644
--- /dev/null
+++ b/Engine/Plugins/VirtualProduction/CaptureManager/CaptureManagerApp/Source/LiveLinkCapabilities/Public/Ingest/IngestCapability_TakeInformation.h
@@ -0,0 +1,33 @@
+// Copyright Epic Games, Inc. All Rights Reserved.
+
+#pragma once
+
+#include "LiveLinkDevice.h"
+
+#include "IngestCapability_TakeInformation.generated.h"
+
+namespace UE::CaptureManager
+{
+using FTakeId = int32;
+}
+
+UCLASS()
+class LIVELINKCAPABILITIES_API UIngestCapability_TakeInformation : public UObject
+{
+	GENERATED_BODY()
+
+public:
+
+	UPROPERTY()
+	FString DeviceName;
+
+	UPROPERTY()
+	FString SlateName;
+
+	UPROPERTY()
+	int32 TakeNumber = -1;
+
+	UPROPERTY()
+	FDateTime DateTime;
+};
+
diff --git a/Engine/Plugins/VirtualProduction/CaptureManager/CaptureManagerApp/Source/LiveLinkCapabilities/Public/Ingest/LiveLinkDeviceCapability_Ingest.h b/Engine/Plugins/VirtualProduction/CaptureManager/CaptureManagerApp/Source/LiveLinkCapabilities/Public/Ingest/LiveLinkDeviceCapability_Ingest.h
new file mode 100644
--- /dev/null
+++ b/Engine/Plugins/VirtualProduction/CaptureManager/CaptureManagerApp/Source/LiveLinkCapabilities/Public/Ingest/LiveLinkDeviceCapability_Ingest.h
@@ -0,0 +1,101 @@
+// Copyright Epic Games, Inc. All Rights Reserved.
+
+#pragma once
+
+#include "LiveLinkDeviceCapability.h"
+
+#include "Templates/ValueOrError.h"
+#include "CaptureManagerTakeMetadata.h"
+
+#include "Ingest/IngestCapability_Options.h"
+#include "Ingest/IngestCapability_ProcessHandle.h"
+#include "Ingest/IngestCapability_UpdateTakeList.h"
+#include "Ingest/IngestCapability_TakeInformation.h"
+#include "Ingest/IngestCapability_Events.h"
+
+#include "LiveLinkDeviceCapability_Ingest.generated.h"
+
+UINTERFACE(Blueprintable)
+class LIVELINKCAPABILITIES_API ULiveLinkDeviceCapability_Ingest : 
+	public ULiveLinkDeviceCapability
+{
+	GENERATED_BODY()
+
+public:
+	const FName Column_IngestSupport;
+
+	ULiveLinkDeviceCapability_Ingest()
+		: Column_IngestSupport(RegisterTableColumn("IngestSupport"))
+	{
+	}
+
+	virtual SHeaderRow::FColumn::FArguments& GenerateHeaderForColumn(const FName InColumnId, SHeaderRow::FColumn::FArguments& InArgs) const override;
+	virtual TSharedPtr<SWidget> GenerateWidgetForColumn(const FName InColumnId, ULiveLinkDevice* InDevice) const override;
+};
+
+class LIVELINKCAPABILITIES_API ILiveLinkDeviceCapability_Ingest 
+	: public ILiveLinkDeviceCapability
+#if CPP
+	, public UE::CaptureManager::FCaptureEventSource
+#endif
+{
+	GENERATED_BODY()
+
+public:
+
+	ILiveLinkDeviceCapability_Ingest();
+
+	UFUNCTION(BlueprintCallable, BlueprintNativeEvent, Category = "Live Link Device|Ingest")
+	UIngestCapability_ProcessHandle* PrepareIngestProcess(int32 InTakeId);
+
+	UFUNCTION(BlueprintCallable, BlueprintNativeEvent, Category = "Live Link Device|Ingest")
+	void UpdateTakeListAsync();
+
+	UFUNCTION(BlueprintCallable, BlueprintNativeEvent, Category = "Live Link Device|Ingest")
+	UIngestCapability_TakeInformation* GetTakeInformation(int32 InTakeId) const;
+
+	UFUNCTION(BlueprintCallable, BlueprintNativeEvent, Category = "Live Link Device|Ingest")
+	TArray<int32> GetTakeIdentifiers() const;
+
+	UFUNCTION(BlueprintCallable, BlueprintNativeEvent, Category = "Live Link Device|Ingest")
+	UIngestCapability_UpdateTakeListCallback* CreateUpdateTakeListCallback();
+
+	UIngestCapability_ProcessHandle* PrepareIngestProcess_Implementation(int32 InTakeId);
+	virtual void UpdateTakeListAsync_Implementation() PURE_VIRTUAL(ILiveLinkDeviceCapability_Ingest::UpdateTakeListAsync_Implementation);
+	UIngestCapability_TakeInformation* GetTakeInformation_Implementation(int32 InTakeId) const;
+	TArray<int32> GetTakeIdentifiers_Implementation() const;
+	UIngestCapability_UpdateTakeListCallback* CreateUpdateTakeListCallback_Implementation();
+
+	TOptional<FTakeMetadata> GetTakeMetadata(int32 InTakeId) const;
+
+protected:
+
+	void ExecuteProcessFinishedReporter(int32 InTakeIdentifier, TValueOrError<void, FIngestCapability_Error> InMaybeError);
+	void ExecuteProcessProgressReporter(int32 InTakeIdentifier, double InProgress);
+	void ExecuteUpdateTakeListCallback(const TArray<int32>& InTakeIdentifiers);
+	
+	int32 AddTake(FTakeMetadata InTakeMetadata);
+	void RemoveTake(int32 InTakeId);
+	void RemoveAllTakes();
+	bool UpdateTake(int32 InTakeId, FTakeMetadata InTakeMetadata);
+
+private:
+
+	static TStrongObjectPtr<UIngestCapability_ProcessHandle>
+		CreateIngestProcessHandle(int32 InTakeId, ILiveLinkDeviceCapability_Ingest* InOwner);
+
+	virtual void CancelIngest(int32 InTakeId) = 0;
+	virtual void RunIngestTakeAsync(int32 InTakeId, const UIngestCapability_Options* InIngestOptions) = 0;
+
+	mutable FCriticalSection TakeMetadataMutex;
+	TMap<int32, FTakeMetadata> TakeMetadataMap;
+
+	mutable FCriticalSection IngestProcessMutex;
+	TMap<int32, TStrongObjectPtr<UIngestCapability_ProcessHandle>> IngestProcessHandles;
+
+	TStrongObjectPtr<UIngestCapability_UpdateTakeListCallback> UpdateTakeListCallback;
+
+	std::atomic_int32_t CurrentTakeId = 0;
+
+	friend class UIngestCapability_ProcessHandle;
+};
\ No newline at end of file
diff --git a/Engine/Plugins/VirtualProduction/CaptureManager/CaptureManagerApp/Source/LiveLinkHubCaptureSourceWidgets/LiveLinkHubCaptureSourceWidgets.Build.cs b/Engine/Plugins/VirtualProduction/CaptureManager/CaptureManagerApp/Source/LiveLinkHubCaptureSourceWidgets/LiveLinkHubCaptureSourceWidgets.Build.cs
new file mode 100644
--- /dev/null
+++ b/Engine/Plugins/VirtualProduction/CaptureManager/CaptureManagerApp/Source/LiveLinkHubCaptureSourceWidgets/LiveLinkHubCaptureSourceWidgets.Build.cs
@@ -0,0 +1,28 @@
+// Copyright Epic Games, Inc. All Rights Reserved.
+
+using UnrealBuildTool;
+using System.IO;
+
+public class LiveLinkHubCaptureSourceWidgets : ModuleRules
+{
+	public LiveLinkHubCaptureSourceWidgets(ReadOnlyTargetRules Target) : base(Target)
+	{
+		PCHUsage = PCHUsageMode.UseExplicitOrSharedPCHs;
+
+		PublicDependencyModuleNames.AddRange(new string[]
+		{
+			"Core",
+			"CoreUObject",
+			"Engine",
+		});
+
+		PrivateDependencyModuleNames.AddRange(new string[]
+		{
+			"CaptureSourceFramework",
+			"InputCore",
+			"Slate",
+			"SlateCore",
+			"Sockets",
+		});
+	}
+}
diff --git a/Engine/Plugins/VirtualProduction/CaptureManager/CaptureManagerApp/Source/LiveLinkHubCaptureSourceWidgets/Private/LiveLinkHubCaptureSourceWidgetsModule.cpp b/Engine/Plugins/VirtualProduction/CaptureManager/CaptureManagerApp/Source/LiveLinkHubCaptureSourceWidgets/Private/LiveLinkHubCaptureSourceWidgetsModule.cpp
new file mode 100644
--- /dev/null
+++ b/Engine/Plugins/VirtualProduction/CaptureManager/CaptureManagerApp/Source/LiveLinkHubCaptureSourceWidgets/Private/LiveLinkHubCaptureSourceWidgetsModule.cpp
@@ -0,0 +1,13 @@
+// Copyright Epic Games, Inc. All Rights Reserved.
+
+#include "LiveLinkHubCaptureSourceWidgetsModule.h"
+
+void FLiveLinkHubCaptureSourceWidgetsModule::StartupModule()
+{
+}
+
+void FLiveLinkHubCaptureSourceWidgetsModule::ShutdownModule()
+{
+}
+
+IMPLEMENT_MODULE(FLiveLinkHubCaptureSourceWidgetsModule, LiveLinkHubCaptureSourceWidgets)
\ No newline at end of file
diff --git a/Engine/Plugins/VirtualProduction/CaptureManager/CaptureManagerApp/Source/LiveLinkHubCaptureSourceWidgets/Private/LiveLinkHubCaptureSourceWidgetsModule.h b/Engine/Plugins/VirtualProduction/CaptureManager/CaptureManagerApp/Source/LiveLinkHubCaptureSourceWidgets/Private/LiveLinkHubCaptureSourceWidgetsModule.h
new file mode 100644
--- /dev/null
+++ b/Engine/Plugins/VirtualProduction/CaptureManager/CaptureManagerApp/Source/LiveLinkHubCaptureSourceWidgets/Private/LiveLinkHubCaptureSourceWidgetsModule.h
@@ -0,0 +1,12 @@
+// Copyright Epic Games, Inc. All Rights Reserved.
+
+#pragma once
+
+#include "Modules/ModuleManager.h"
+
+class FLiveLinkHubCaptureSourceWidgetsModule: public IModuleInterface
+{
+public:
+	virtual void StartupModule() override;
+	virtual void ShutdownModule() override;
+};
\ No newline at end of file
diff --git a/Engine/Plugins/VirtualProduction/CaptureManager/CaptureManagerApp/Source/LiveLinkHubCaptureSourceWidgets/Private/SArchiveCreationPanel.cpp b/Engine/Plugins/VirtualProduction/CaptureManager/CaptureManagerApp/Source/LiveLinkHubCaptureSourceWidgets/Private/SArchiveCreationPanel.cpp
new file mode 100644
--- /dev/null
+++ b/Engine/Plugins/VirtualProduction/CaptureManager/CaptureManagerApp/Source/LiveLinkHubCaptureSourceWidgets/Private/SArchiveCreationPanel.cpp
@@ -0,0 +1,168 @@
+// Copyright Epic Games, Inc. All Rights Reserved.
+
+#include "Widgets/SArchiveCreationPanel.h"
+
+#include "Utils.h"
+
+#include "CaptureSourceFrameworkModule.h"
+
+#include "Framework/Application/SlateApplication.h"
+#include "Widgets/Input/SButton.h"
+#include "Widgets/Input/SEditableTextBox.h"
+#include "Widgets/Layout/SGridPanel.h"
+#include "Widgets/Text/STextBlock.h"
+
+#include "DesktopPlatformModule.h"
+#include "HAL/FileManager.h"
+#include "Misc/Paths.h"
+
+#define LOCTEXT_NAMESPACE "LiveLinkHubCaptureSourceWidgets_SArchiveCreationPanel"
+
+namespace UE::CaptureManager
+{
+
+const FText SArchiveCreationPanel::BrowseToLocationText = LOCTEXT("BrowseToArhiveLocation", "Browse to archive location");
+const FText SArchiveCreationPanel::CreateText = LOCTEXT("Create", "Create...");
+const FText SArchiveCreationPanel::NameText = LOCTEXT("Name", "Name");
+
+TMap<FString, FString> SArchiveCreationPanel::LastLocationByFactoryId = {};
+// We protect the last location map with a mutex out of an abundance of caution
+FCriticalSection SArchiveCreationPanel::LastLocationCriticalSection;
+
+FReply SArchiveCreationPanel::CreateSource()
+{
+	IDesktopPlatform* DesktopPlatform = FDesktopPlatformModule::Get();
+	check(DesktopPlatform);
+
+	FString ArchiveFolderPath;
+
+	const bool bIsSelected = DesktopPlatform->OpenDirectoryDialog(
+		FSlateApplication::Get().FindBestParentWindowHandleForDialogs(nullptr),
+		BrowseToLocationText.ToString(),
+		GetLastLocation(),
+		ArchiveFolderPath
+	);
+
+	if (bIsSelected)
+	{
+		TMap<FString, FPropertyValue> ParamValues;
+		ParamValues.Emplace(TEXT("StoragePath"), MakePropertyValue<FString>(ArchiveFolderPath));
+
+		const FString CaptureSourceName = SourceNameEntryBox->GetText().ToString();
+		SetLastLocation(ArchiveFolderPath);
+
+		OnSourceConfigured.ExecuteIfBound(FactoryId, CaptureSourceName, ParamValues);
+	}
+
+	return FReply::Handled();
+}
+
+void SArchiveCreationPanel::OnValidityChanged()
+{
+	CreateButton->SetEnabled(bSourceNameIsValid);
+}
+
+bool SArchiveCreationPanel::OnSourceNameChanged(const FText& InText, FText& OutErrorMessage)
+{
+	bSourceNameIsValid = ValidateCaptureSourceName(InText, OutErrorMessage);
+	OnValidityChanged();
+	return bSourceNameIsValid;
+}
+
+// We override the default behaviour when the enter key is pressed on the source name widget, to keep the creation panel open. 
+// We instead forward the user's focus onto the next widget.
+FReply SArchiveCreationPanel::OnSourceNameKeyDown(const FGeometry& MyGeometry, const FKeyEvent& KeyEvent)
+{
+	if (KeyEvent.GetKey() == EKeys::Enter)
+	{
+		return FReply::Handled().SetUserFocus(CreateButton.ToSharedRef(), EFocusCause::Navigation);
+	}
+
+	return FReply::Unhandled();
+}
+
+void SArchiveCreationPanel::Construct(const FArguments& InArgs)
+{
+	FactoryId = InArgs._FactoryId;
+	OnSourceConfigured = InArgs._OnSourceConfigured;
+
+	constexpr float MinDesiredWidth = 100.0f;
+	constexpr float GridPanelPadding = 2.0f;
+
+	CreateButton = SNew(SButton)
+		.HAlign(EHorizontalAlignment::HAlign_Center)
+		.VAlign(EVerticalAlignment::VAlign_Center)
+		.ContentPadding(FMargin(6, 2))
+		.Text(CreateText)
+		.OnClicked(this, &SArchiveCreationPanel::CreateSource);
+
+	SourceNameEntryBox = SNew(SEditableTextBox)
+		.MinDesiredWidth(MinDesiredWidth)
+		.OnVerifyTextChanged(this, &SArchiveCreationPanel::OnSourceNameChanged)
+		.OnKeyDownHandler(this, &SArchiveCreationPanel::OnSourceNameKeyDown);
+
+	ChildSlot
+		[
+			SNew(SGridPanel)
+				+ SGridPanel::Slot(0, 0)
+				.Padding(GridPanelPadding)
+				.VAlign(EVerticalAlignment::VAlign_Center)
+				[
+					SNew(STextBlock)
+						.Text(NameText)
+				]
+				+ SGridPanel::Slot(1, 0)
+				.VAlign(EVerticalAlignment::VAlign_Center)
+				.Padding(GridPanelPadding)
+				[
+					SourceNameEntryBox.ToSharedRef()
+				]
+				+ SGridPanel::Slot(1, 1)
+				.VAlign(EVerticalAlignment::VAlign_Center)
+				.Padding(GridPanelPadding)
+				[
+					CreateButton.ToSharedRef()
+				]
+		];
+
+	// Update the starting validity
+	FText ErrorMessage;
+
+	if (OnSourceNameChanged(SourceNameEntryBox->GetText(), ErrorMessage))
+	{
+		SourceNameEntryBox->SetError(ErrorMessage);
+	}
+}
+
+FString SArchiveCreationPanel::GetLastLocation() const
+{
+	FScopeLock Lock(&LastLocationCriticalSection);
+
+	const FString* LastLocation = LastLocationByFactoryId.Find(FactoryId);
+
+	if (LastLocation)
+	{
+		return *LastLocation;
+	}
+
+	// Sensible default location
+	return IFileManager::Get().ConvertToAbsolutePathForExternalAppForWrite(*FPaths::ProjectContentDir());
+}
+
+void SArchiveCreationPanel::SetLastLocation(const FString& Location)
+{
+	FScopeLock Lock(&LastLocationCriticalSection);
+
+	if (FString* LastLocation = LastLocationByFactoryId.Find(FactoryId))
+	{
+		*LastLocation = Location;
+	}
+	else
+	{
+		LastLocationByFactoryId.Emplace(FactoryId, Location);
+	}
+}
+
+} // UE::CaptureManager
+
+#undef LOCTEXT_NAMESPACE
\ No newline at end of file
diff --git a/Engine/Plugins/VirtualProduction/CaptureManager/CaptureManagerApp/Source/LiveLinkHubCaptureSourceWidgets/Private/SConnectionCreationPanel.cpp b/Engine/Plugins/VirtualProduction/CaptureManager/CaptureManagerApp/Source/LiveLinkHubCaptureSourceWidgets/Private/SConnectionCreationPanel.cpp
new file mode 100644
--- /dev/null
+++ b/Engine/Plugins/VirtualProduction/CaptureManager/CaptureManagerApp/Source/LiveLinkHubCaptureSourceWidgets/Private/SConnectionCreationPanel.cpp
@@ -0,0 +1,278 @@
+// Copyright Epic Games, Inc. All Rights Reserved.
+
+#include "Widgets/SConnectionCreationPanel.h"
+#include "Utils.h"
+
+#include "CaptureSourceFrameworkModule.h"
+
+#include "Widgets/Input/SButton.h"
+#include "Widgets/Input/SEditableTextBox.h"
+#include "Widgets/Input/SSpinBox.h"
+#include "Widgets/Layout/SGridPanel.h"
+
+#include "IPAddress.h"
+#include "SocketSubsystem.h" 
+
+#define LOCTEXT_NAMESPACE "LiveLinkHubCaptureSourceWidgets_SConnectionCreationPanel"
+
+namespace UE::CaptureManager
+{
+
+static bool ValidateIpAddress(const FString& InIpAddress)
+{
+	if (InIpAddress.IsEmpty())
+	{
+		return false;
+	}
+
+	ISocketSubsystem* SocketSubsystem = ISocketSubsystem::Get(PLATFORM_SOCKETSUBSYSTEM);
+
+	if (!SocketSubsystem)
+	{
+		return false;
+	}
+
+	TSharedPtr<FInternetAddr> InternetAddr = SocketSubsystem->CreateInternetAddr();
+
+	if (!InternetAddr)
+	{
+		return false;
+	}
+
+	bool bIsValidIP = false;
+	InternetAddr->SetIp(*InIpAddress, bIsValidIP);
+
+	return bIsValidIP;
+}
+
+// We have to customize the spin box as it does not have an OnKeyDown handler we can easily hook into. So intead we need to override the 
+// OnPreviewKeyDown method of the spin box.
+class SPortNumberSpinBox : public SSpinBox<uint16>
+{
+public:
+	SLATE_BEGIN_ARGS(SPortNumberSpinBox)
+		{}
+		SLATE_ARGUMENT(TSharedPtr<SWidget>, WidgetToFocus)
+		SLATE_ARGUMENT(uint16, DefaultPort)
+		SLATE_ARGUMENT(uint16, MinPortNumber)
+	SLATE_END_ARGS()
+
+	SPortNumberSpinBox() : SSpinBox<uint16>()
+	{
+	}
+
+	void Construct(const FArguments& InArgs)
+	{
+		WidgetToFocus = InArgs._WidgetToFocus;
+
+		SSpinBox<uint16>::Construct(
+			SSpinBox<uint16>::FArguments()
+			.MinValue(InArgs._MinPortNumber)
+		);
+
+		SetValue(InArgs._DefaultPort);
+	}
+
+private:
+	// We override the default behaviour when the enter key is pressed on the port number widget, to keep the creation panel open. 
+	// We instead forward the user's focus onto the next widget.
+	virtual FReply OnPreviewKeyDown(const FGeometry& MyGeometry, const FKeyEvent& InKeyEvent) override
+	{
+		if (WidgetToFocus && InKeyEvent.GetKey() == EKeys::Enter)
+		{
+			return FReply::Handled().SetUserFocus(WidgetToFocus.ToSharedRef(), EFocusCause::Navigation);
+		}
+
+		return SSpinBox<uint16>::OnPreviewKeyDown(MyGeometry, InKeyEvent);
+	}
+
+	TSharedPtr<SWidget> WidgetToFocus;
+};
+
+const FText SConnectionCreationPanel::CreateText = LOCTEXT("Create", "Create");
+const FText SConnectionCreationPanel::NameText = LOCTEXT("Name", "Name");
+const FText SConnectionCreationPanel::IpAddressText = LOCTEXT("IpAddress", "IP Address");
+const FText SConnectionCreationPanel::PortText = LOCTEXT("Port", "Port");
+
+FReply SConnectionCreationPanel::CreateSource()
+{
+	const FString CaptureSourceName = SourceNameEntryBox->GetText().ToString();
+	const uint16 PortNumber = PortNumberEntryBox->GetValue();
+	FString IpAddress = IpAddressEntryBox->GetText().ToString();
+
+	TMap<FString, FPropertyValue> ParamValues;
+	ParamValues.Emplace(TEXT("IpAddress"), MakePropertyValue<FString>(IpAddress));
+	ParamValues.Emplace(TEXT("Port"), MakePropertyValue<int64>(PortNumber));
+
+	OnSourceConfigured.ExecuteIfBound(FactoryId, CaptureSourceName, ParamValues);
+
+	// Reset the name to make it clear that we no longer have a valid name (it would be a duplicate)
+	SourceNameEntryBox->SetText(FText::GetEmpty());
+	OnValidityChanged();
+
+	return FReply::Handled();
+}
+
+void SConnectionCreationPanel::OnValidityChanged()
+{
+	CreateButton->SetEnabled(IsReadyToCreate());
+}
+
+bool SConnectionCreationPanel::IsReadyToCreate() const
+{
+	return bSourceNameIsValid && bIpAddressIsValid;
+}
+
+bool SConnectionCreationPanel::OnSourceNameChanged(const FText& InSourceName, FText& OutErrorMessage)
+{
+	bSourceNameIsValid = ValidateCaptureSourceName(InSourceName, OutErrorMessage);
+	OnValidityChanged();
+	return bSourceNameIsValid;
+}
+
+void SConnectionCreationPanel::OnIpAddressCommitted(const FText& InIpAddress, ETextCommit::Type InCommitType)
+{
+	bIpAddressIsValid = ValidateIpAddress(InIpAddress.ToString());
+	OnValidityChanged();
+
+	if (!bIpAddressIsValid)
+	{
+		IpAddressEntryBox->SetError(TEXT("Invalid IP address"));
+	}
+}
+
+// We override the default behaviour when the enter key is pressed on the source name widget, to keep the creation panel open. 
+// We instead forward the user's focus onto the next widget.
+FReply SConnectionCreationPanel::OnSourceNameKeyDown(const FGeometry& MyGeometry, const FKeyEvent& KeyEvent)
+{
+	if (KeyEvent.GetKey() == EKeys::Enter)
+	{
+		return FReply::Handled().SetUserFocus(IpAddressEntryBox.ToSharedRef(), EFocusCause::Navigation);
+	}
+
+	return FReply::Unhandled();
+}
+
+// We override the default behaviour when the enter key is pressed on the ip address widget, to keep the creation panel open. 
+// We instead forward the user's focus onto the next widget.
+FReply SConnectionCreationPanel::OnIpAddressKeyDown(const FGeometry& MyGeometry, const FKeyEvent& KeyEvent)
+{
+	if (KeyEvent.GetKey() == EKeys::Enter)
+	{
+		return FReply::Handled().SetUserFocus(PortNumberEntryBox.ToSharedRef(), EFocusCause::Navigation);
+	}
+
+	return FReply::Unhandled();
+}
+
+void SConnectionCreationPanel::Construct(const FArguments& InArgs)
+{
+	constexpr float MinTextBoxWidth = 100.0f;
+	constexpr float GridPanelPadding = 2.0f;
+
+	FactoryId = InArgs._FactoryId;
+	OnSourceConfigured = InArgs._OnSourceConfigured;
+
+	CreateButton = SNew(SButton)
+		.HAlign(EHorizontalAlignment::HAlign_Center)
+		.VAlign(EVerticalAlignment::VAlign_Center)
+		.ContentPadding(FMargin(6, 2))
+		.Text(CreateText)
+		.OnClicked(this, &SConnectionCreationPanel::CreateSource);
+
+	// For whatever reason, the default foreground color for the text box does not match the spin box, which 
+	// makes it look like the spin box has focus, so here we just make them the same.
+	const static FEditableTextBoxStyle InfoWidgetStyle =
+		FEditableTextBoxStyle(FAppStyle::Get().GetWidgetStyle<FEditableTextBoxStyle>("NormalEditableTextBox"))
+		.SetForegroundColor(FAppStyle::Get().GetWidgetStyle<FSpinBoxStyle>("SpinBox").ForegroundColor);
+
+	SourceNameEntryBox = SNew(SEditableTextBox)
+		.MinDesiredWidth(MinTextBoxWidth)
+		.OnVerifyTextChanged(this, &SConnectionCreationPanel::OnSourceNameChanged)
+		.Style(&InfoWidgetStyle)
+		.OnKeyDownHandler(this, &SConnectionCreationPanel::OnSourceNameKeyDown);
+
+	IpAddressEntryBox = SNew(SEditableTextBox)
+		.MinDesiredWidth(MinTextBoxWidth)
+		// We only check the ip address on commit, as it logs every error to the console
+		.OnTextCommitted(this, &SConnectionCreationPanel::OnIpAddressCommitted)
+		.Style(&InfoWidgetStyle)
+		.OnKeyDownHandler(this, &SConnectionCreationPanel::OnIpAddressKeyDown);
+
+	// Try to set up some useful defaults, mostly for dev testing purposes
+	IpAddressEntryBox->SetText(FText::FromString(TEXT("127.0.0.1")));
+
+	PortNumberEntryBox = SNew(SPortNumberSpinBox)
+		.WidgetToFocus(CreateButton)
+		.DefaultPort(14785)
+		.MinPortNumber(1025);
+
+	ChildSlot
+		[
+			SNew(SGridPanel)
+				// First row
+				+ SGridPanel::Slot(0, 0)
+				.Padding(GridPanelPadding)
+				.HAlign(EHorizontalAlignment::HAlign_Left)
+				.VAlign(EVerticalAlignment::VAlign_Center)
+				[
+					SNew(STextBlock)
+						.Text(NameText)
+				]
+				+ SGridPanel::Slot(0, 1)
+				.Padding(GridPanelPadding)
+				.HAlign(EHorizontalAlignment::HAlign_Left)
+				.VAlign(EVerticalAlignment::VAlign_Center)
+				[
+					SNew(STextBlock)
+						.Text(IpAddressText)
+				]
+				+ SGridPanel::Slot(0, 2)
+				.Padding(GridPanelPadding)
+				.HAlign(EHorizontalAlignment::HAlign_Left)
+				.VAlign(EVerticalAlignment::VAlign_Center)
+				[
+					SNew(STextBlock)
+						.Text(PortText)
+				]
+				// Second row
+				+ SGridPanel::Slot(1, 0)
+				.Padding(GridPanelPadding)
+				.VAlign(EVerticalAlignment::VAlign_Center)
+				[
+					SourceNameEntryBox.ToSharedRef()
+				]
+				+ SGridPanel::Slot(1, 1)
+				.Padding(GridPanelPadding)
+				.VAlign(EVerticalAlignment::VAlign_Center)
+				[
+					IpAddressEntryBox.ToSharedRef()
+				]
+				+ SGridPanel::Slot(1, 2)
+				.Padding(GridPanelPadding)
+				.VAlign(EVerticalAlignment::VAlign_Center)
+				[
+					PortNumberEntryBox.ToSharedRef()
+				]
+				+ SGridPanel::Slot(1, 3)
+				.VAlign(EVerticalAlignment::VAlign_Center)
+				.Padding(GridPanelPadding)
+				[
+					CreateButton.ToSharedRef()
+				]
+		];
+
+	// Update the starting validity
+	FText ErrorMessage;
+
+	if (OnSourceNameChanged(SourceNameEntryBox->GetText(), ErrorMessage))
+	{
+		SourceNameEntryBox->SetError(ErrorMessage.ToString());
+	}
+
+	OnIpAddressCommitted(IpAddressEntryBox->GetText(), ETextCommit::Default);
+}
+
+} // UE::CaptureManager
+
+#undef LOCTEXT_NAMESPACE
\ No newline at end of file
diff --git a/Engine/Plugins/VirtualProduction/CaptureManager/CaptureManagerApp/Source/LiveLinkHubCaptureSourceWidgets/Private/Utils.cpp b/Engine/Plugins/VirtualProduction/CaptureManager/CaptureManagerApp/Source/LiveLinkHubCaptureSourceWidgets/Private/Utils.cpp
new file mode 100644
--- /dev/null
+++ b/Engine/Plugins/VirtualProduction/CaptureManager/CaptureManagerApp/Source/LiveLinkHubCaptureSourceWidgets/Private/Utils.cpp
@@ -0,0 +1,60 @@
+// Copyright Epic Games, Inc. All Rights Reserved.
+
+#include "Utils.h"
+
+#include "CaptureSourceFrameworkModule.h"
+
+#define LOCTEXT_NAMESPACE "LiveLinkHubCaptureSourceWidgets_Utils"
+
+DEFINE_LOG_CATEGORY(LogLiveLinkHubCaptureSourceWidgets);
+
+namespace UE::CaptureManager
+{
+static const FText AlreadyExistsText = LOCTEXT("AlreadyExists", "Already exists");
+
+static bool SourceNameAlreadyExists(const FString& SourceName)
+{
+	FCaptureSourceFrameworkModule& Module = FModuleManager::LoadModuleChecked<FCaptureSourceFrameworkModule>("CaptureSourceFramework");
+	FCaptureSourceManager& CaptureSourceManager = Module.GetCaptureSourceManager();
+
+	for (const FCaptureSourceId& CaptureSourceId : CaptureSourceManager.GetCaptureSourceIds())
+	{
+		const FCaptureSourceHandle CaptureSource = CaptureSourceManager.GetCaptureSource(CaptureSourceId);
+
+		if (CaptureSource.GetName() == SourceName)
+		{
+			return true;
+		}
+	}
+
+	return false;
+}
+
+bool ValidateCaptureSourceName(const FText& InSourceName, FText& OutErrorMessage)
+{
+	if (InSourceName.IsEmpty())
+	{
+		// Don't populate the error message, this is a common enough situation that reporting it as an error would be annoying, it is still invalid 
+		// but the user already knows this, as the creation button gets disabled as a result.
+		return false;
+	}
+
+	const FString SourceName = InSourceName.ToString();
+
+	if (!FName::IsValidXName(SourceName, INVALID_OBJECTNAME_CHARACTERS INVALID_LONGPACKAGE_CHARACTERS, &OutErrorMessage, &InSourceName))
+	{
+		return false;
+	}
+
+	if (SourceNameAlreadyExists(SourceName))
+	{
+		OutErrorMessage = AlreadyExistsText;
+		return false;
+	}
+
+	return true;
+}
+
+}
+
+#undef LOCTEXT_NAMESPACE
\ No newline at end of file
diff --git a/Engine/Plugins/VirtualProduction/CaptureManager/CaptureManagerApp/Source/LiveLinkHubCaptureSourceWidgets/Private/Utils.h b/Engine/Plugins/VirtualProduction/CaptureManager/CaptureManagerApp/Source/LiveLinkHubCaptureSourceWidgets/Private/Utils.h
new file mode 100644
--- /dev/null
+++ b/Engine/Plugins/VirtualProduction/CaptureManager/CaptureManagerApp/Source/LiveLinkHubCaptureSourceWidgets/Private/Utils.h
@@ -0,0 +1,12 @@
+// Copyright Epic Games, Inc. All Rights Reserved.
+
+#pragma once
+
+#include "Logging/LogMacros.h"
+
+DECLARE_LOG_CATEGORY_EXTERN(LogLiveLinkHubCaptureSourceWidgets, Log, All);
+
+namespace UE::CaptureManager
+{
+bool ValidateCaptureSourceName(const FText& InSourceName, FText& OutErrorMessage);
+}
diff --git a/Engine/Plugins/VirtualProduction/CaptureManager/CaptureManagerApp/Source/LiveLinkHubCaptureSourceWidgets/Public/Widgets/SArchiveCreationPanel.h b/Engine/Plugins/VirtualProduction/CaptureManager/CaptureManagerApp/Source/LiveLinkHubCaptureSourceWidgets/Public/Widgets/SArchiveCreationPanel.h
new file mode 100644
--- /dev/null
+++ b/Engine/Plugins/VirtualProduction/CaptureManager/CaptureManagerApp/Source/LiveLinkHubCaptureSourceWidgets/Public/Widgets/SArchiveCreationPanel.h
@@ -0,0 +1,54 @@
+// Copyright Epic Games, Inc. All Rights Reserved.
+
+#pragma once
+
+#include "Templates/SharedPointer.h"
+#include "Widgets/SCompoundWidget.h"
+
+#include "Widgets/SDefaultCaptureSourceCreator.h"
+
+class SButton;
+class SEditableTextBox;
+class STextBlock;
+
+namespace UE::CaptureManager
+{
+
+class LIVELINKHUBCAPTURESOURCEWIDGETS_API SArchiveCreationPanel : public SCompoundWidget
+{
+public:
+	SLATE_BEGIN_ARGS(SArchiveCreationPanel) {}
+		SLATE_ARGUMENT(FString, FactoryId)
+		SLATE_ARGUMENT(FOnCaptureSourceConfigured, OnSourceConfigured)
+	SLATE_END_ARGS()
+
+	void Construct(const FArguments& InArgs);
+
+private:
+	FReply CreateSource();
+
+	void OnValidityChanged();
+	bool OnSourceNameChanged(const FText& InText, FText& OutErrorMessage);
+	FReply OnSourceNameKeyDown(const FGeometry& MyGeometry, const FKeyEvent& KeyEvent);
+
+	FString GetLastLocation() const;
+	void SetLastLocation(const FString& Location);
+
+	TSharedPtr<SButton> CreateButton;
+	TSharedPtr<SEditableTextBox> SourceNameEntryBox;
+	FString FactoryId;
+	FOnCaptureSourceConfigured OnSourceConfigured;
+
+	bool bSourceNameIsValid = false;
+
+	static const FText BrowseToLocationText;
+	static const FText CreateText;
+	static const FText NameText;
+
+	// A place to persist the last browse location, at least for the duration of the program.
+	static TMap<FString, FString> LastLocationByFactoryId;
+	// We protect the last location map with a mutex out of an abundance of caution
+	static FCriticalSection LastLocationCriticalSection;
+};
+
+}
diff --git a/Engine/Plugins/VirtualProduction/CaptureManager/CaptureManagerApp/Source/LiveLinkHubCaptureSourceWidgets/Public/Widgets/SConnectionCreationPanel.h b/Engine/Plugins/VirtualProduction/CaptureManager/CaptureManagerApp/Source/LiveLinkHubCaptureSourceWidgets/Public/Widgets/SConnectionCreationPanel.h
new file mode 100644
--- /dev/null
+++ b/Engine/Plugins/VirtualProduction/CaptureManager/CaptureManagerApp/Source/LiveLinkHubCaptureSourceWidgets/Public/Widgets/SConnectionCreationPanel.h
@@ -0,0 +1,54 @@
+// Copyright Epic Games, Inc. All Rights Reserved.
+
+#pragma once
+
+#include "Templates/SharedPointer.h"
+#include "Widgets/SCompoundWidget.h"
+
+#include "Widgets/SDefaultCaptureSourceCreator.h"
+
+class SButton;
+class SEditableTextBox;
+
+namespace UE::CaptureManager
+{
+class SPortNumberSpinBox;
+
+class LIVELINKHUBCAPTURESOURCEWIDGETS_API SConnectionCreationPanel : public SCompoundWidget
+{
+public:
+	SLATE_BEGIN_ARGS(SConnectionCreationPanel) {}
+		SLATE_ARGUMENT(FString, FactoryId)
+		SLATE_ARGUMENT(FOnCaptureSourceConfigured, OnSourceConfigured)
+	SLATE_END_ARGS()
+
+	void Construct(const FArguments& InArgs);
+
+private:
+	FReply CreateSource();
+	void OnValidityChanged();
+	bool IsReadyToCreate() const;
+
+	bool OnSourceNameChanged(const FText& InText, FText& OutErrorMessage);
+	void OnIpAddressCommitted(const FText& InIpAddress, ETextCommit::Type InCommitType);
+	FReply OnSourceNameKeyDown(const FGeometry& MyGeometry, const FKeyEvent& KeyEvent);
+	FReply OnIpAddressKeyDown(const FGeometry& MyGeometry, const FKeyEvent& KeyEvent);
+
+	TSharedPtr<SButton> CreateButton;
+	TSharedPtr<SEditableTextBox> SourceNameEntryBox;
+	TSharedPtr<SEditableTextBox> IpAddressEntryBox;
+	TSharedPtr<SPortNumberSpinBox> PortNumberEntryBox;
+
+	FString FactoryId;
+	FOnCaptureSourceConfigured OnSourceConfigured;
+
+	bool bSourceNameIsValid = false;
+	bool bIpAddressIsValid = false;
+
+	static const FText CreateText;
+	static const FText NameText;
+	static const FText IpAddressText;
+	static const FText PortText;
+};
+
+}
diff --git a/Engine/Plugins/VirtualProduction/CaptureManager/CaptureManagerApp/Source/LiveLinkHubDiscovery/LiveLinkHubDiscovery.Build.cs b/Engine/Plugins/VirtualProduction/CaptureManager/CaptureManagerApp/Source/LiveLinkHubDiscovery/LiveLinkHubDiscovery.Build.cs
new file mode 100644
--- /dev/null
+++ b/Engine/Plugins/VirtualProduction/CaptureManager/CaptureManagerApp/Source/LiveLinkHubDiscovery/LiveLinkHubDiscovery.Build.cs
@@ -0,0 +1,22 @@
+// Copyright Epic Games, Inc. All Rights Reserved.
+
+using UnrealBuildTool;
+
+public class LiveLinkHubDiscovery : ModuleRules
+{
+	public LiveLinkHubDiscovery(ReadOnlyTargetRules Target) : base(Target)
+	{
+		PrivateDependencyModuleNames.AddRange(new string[] 
+		{
+			"CaptureSourceFramework",
+			"Core",
+			"CoreUObject",
+			"CaptureUtils",
+			"LiveLinkHubCaptureMessaging",
+			"CaptureCommunicationManager",
+			"CaptureManagerTakeMetadata",
+			"Messaging",
+			"Sockets",
+		});
+	}
+}
diff --git a/Engine/Plugins/VirtualProduction/CaptureManager/CaptureManagerApp/Source/LiveLinkHubDiscovery/Private/LiveLinkHubDiscoveryModule.cpp b/Engine/Plugins/VirtualProduction/CaptureManager/CaptureManagerApp/Source/LiveLinkHubDiscovery/Private/LiveLinkHubDiscoveryModule.cpp
new file mode 100644
--- /dev/null
+++ b/Engine/Plugins/VirtualProduction/CaptureManager/CaptureManagerApp/Source/LiveLinkHubDiscovery/Private/LiveLinkHubDiscoveryModule.cpp
@@ -0,0 +1,20 @@
+// Copyright Epic Games, Inc. All Rights Reserved.
+
+#include "LiveLinkHubDiscoveryModule.h"
+
+void FLiveLinkHubDiscovery::StartupModule()
+{
+	DiscoveryRequester = MakeShared<UE::CaptureManager::FDiscoveryRequester>();
+}
+
+void FLiveLinkHubDiscovery::ShutdownModule()
+{
+	DiscoveryRequester = nullptr;
+}
+
+TSharedRef<UE::CaptureManager::FDiscoveryRequester> FLiveLinkHubDiscovery::GetDiscoveryRequster()
+{
+	return DiscoveryRequester.ToSharedRef();
+}
+
+IMPLEMENT_MODULE(FLiveLinkHubDiscovery, LiveLinkHubDiscovery);
\ No newline at end of file
diff --git a/Engine/Plugins/VirtualProduction/CaptureManager/CaptureManagerApp/Source/LiveLinkHubDiscovery/Public/Discovery/DiscoveryRequester.h b/Engine/Plugins/VirtualProduction/CaptureManager/CaptureManagerApp/Source/LiveLinkHubDiscovery/Public/Discovery/DiscoveryRequester.h
new file mode 100644
--- /dev/null
+++ b/Engine/Plugins/VirtualProduction/CaptureManager/CaptureManagerApp/Source/LiveLinkHubDiscovery/Public/Discovery/DiscoveryRequester.h
@@ -0,0 +1,67 @@
+// Copyright Epic Games, Inc. All Rights Reserved.
+
+#pragma once
+
+#include "MessageEndpoint.h"
+
+DECLARE_LOG_CATEGORY_EXTERN(LogLiveLinkHubDiscovery, Log, All)
+
+namespace UE::CaptureManager
+{
+
+class LIVELINKHUBDISCOVERY_API FDiscoveredClient
+{
+public:
+
+	FDiscoveredClient(FString InHostName, 
+					  FString InIPAddress, 
+					  uint16 InPort, 
+					  FMessageAddress InMessageAddress);
+
+	FDiscoveredClient(const FDiscoveredClient& InOther);
+	FDiscoveredClient(FDiscoveredClient&& InOther);
+
+	bool IsActive() const;
+
+	FString GetHostName() const;
+	FString GetIPAddress() const;
+	uint16 GetExportPort() const;
+	FMessageAddress GetMessageAddress() const;
+
+	void SetLastDiscoveryResponse(double InLastDiscoveryResponse);
+
+	friend bool operator==(const FDiscoveredClient& InLeft, const FDiscoveredClient& InRight);
+
+private:
+
+	static constexpr double InactiveTimeout = 5.0f;
+
+	FString HostName;
+	FString IPAddress;
+	uint16 ExportPort;
+
+	FMessageAddress MessageAddress;
+
+	std::atomic<double> LastDiscoveryResponse = 0.0f;
+};
+
+class LIVELINKHUBDISCOVERY_API FDiscoveryRequester
+{
+public:
+	DECLARE_TS_MULTICAST_DELEGATE_OneParam(FClientDiscovered, const FDiscoveredClient& InClient);
+
+	FDiscoveryRequester();
+	~FDiscoveryRequester();
+
+	TArray<FDiscoveredClient> GetDiscoveredClients() const;
+
+	FDelegateHandle AddClientDiscoveredListener(FClientDiscovered::FDelegate InListener);
+	void RemoveClientDiscoveredListener(FDelegateHandle InListenerHandle);
+
+private:
+
+	struct FImpl;
+	TUniquePtr<FImpl> ImplPtr;
+};
+
+} // namespace UE::CaptureManager
diff --git a/Engine/Plugins/VirtualProduction/CaptureManager/CaptureManagerApp/Source/LiveLinkHubDiscovery/Public/LiveLinkHubDiscoveryModule.h b/Engine/Plugins/VirtualProduction/CaptureManager/CaptureManagerApp/Source/LiveLinkHubDiscovery/Public/LiveLinkHubDiscoveryModule.h
new file mode 100644
--- /dev/null
+++ b/Engine/Plugins/VirtualProduction/CaptureManager/CaptureManagerApp/Source/LiveLinkHubDiscovery/Public/LiveLinkHubDiscoveryModule.h
@@ -0,0 +1,21 @@
+// Copyright Epic Games, Inc. All Rights Reserved.
+
+#pragma once
+
+#include "Modules/ModuleManager.h"
+
+#include "Discovery/DiscoveryRequester.h"
+
+class LIVELINKHUBDISCOVERY_API FLiveLinkHubDiscovery : public IModuleInterface
+{
+public:
+
+	virtual void StartupModule() override;
+	virtual void ShutdownModule() override;
+
+	TSharedRef<UE::CaptureManager::FDiscoveryRequester> GetDiscoveryRequster();
+
+private:
+
+	TSharedPtr<UE::CaptureManager::FDiscoveryRequester> DiscoveryRequester;
+};
diff --git a/Engine/Plugins/VirtualProduction/CaptureManager/CaptureManagerApp/Source/LiveLinkHubExportClient/LiveLinkHubExportClient.Build.cs b/Engine/Plugins/VirtualProduction/CaptureManager/CaptureManagerApp/Source/LiveLinkHubExportClient/LiveLinkHubExportClient.Build.cs
new file mode 100644
--- /dev/null
+++ b/Engine/Plugins/VirtualProduction/CaptureManager/CaptureManagerApp/Source/LiveLinkHubExportClient/LiveLinkHubExportClient.Build.cs
@@ -0,0 +1,28 @@
+// Copyright Epic Games, Inc. All Rights Reserved.
+
+using UnrealBuildTool;
+using System.IO;
+
+public class LiveLinkHubExportClient : ModuleRules
+{
+	public LiveLinkHubExportClient(ReadOnlyTargetRules Target) : base(Target)
+	{
+		PCHUsage = PCHUsageMode.UseExplicitOrSharedPCHs;
+
+		PublicDependencyModuleNames.AddRange(new string[]
+		{
+			"Core",
+			"CoreUObject",
+			"Engine",
+			"CaptureUtils"
+		});
+
+		PrivateDependencyModuleNames.AddRange(new string[]
+		{
+			"UnrealEd",
+			"Sockets",
+			"CaptureManagerTakeMetadata",
+			"LiveLinkHubCaptureMessaging"
+		});
+	}
+}
diff --git a/Engine/Plugins/VirtualProduction/CaptureManager/CaptureManagerApp/Source/LiveLinkHubExportClient/Private/LiveLinkHubExportClientModule.cpp b/Engine/Plugins/VirtualProduction/CaptureManager/CaptureManagerApp/Source/LiveLinkHubExportClient/Private/LiveLinkHubExportClientModule.cpp
new file mode 100644
--- /dev/null
+++ b/Engine/Plugins/VirtualProduction/CaptureManager/CaptureManagerApp/Source/LiveLinkHubExportClient/Private/LiveLinkHubExportClientModule.cpp
@@ -0,0 +1,7 @@
+// Copyright Epic Games, Inc. All Rights Reserved.
+
+#include "Modules/ModuleManager.h"
+
+using FLiveLinkHubExportClientModule = IModuleInterface;
+
+IMPLEMENT_MODULE(FLiveLinkHubExportClientModule, LiveLinkHubExportClient);
\ No newline at end of file
diff --git a/Engine/Plugins/VirtualProduction/CaptureManager/CaptureManagerCore/CaptureManagerCore.uplugin b/Engine/Plugins/VirtualProduction/CaptureManager/CaptureManagerCore/CaptureManagerCore.uplugin
new file mode 100644
--- /dev/null
+++ b/Engine/Plugins/VirtualProduction/CaptureManager/CaptureManagerCore/CaptureManagerCore.uplugin
@@ -0,0 +1,64 @@
+{
+	"FileVersion": 1,
+	"Version": 1,
+	"VersionName": "1.0.0",
+	"FriendlyName": "Capture Manager Core",
+	"Description": "The Capture Manager Core plugin contains utility modules that are shared between Capture Manager App plugin and Capture Manager Editor plugin.",
+	"Category": "Virtual Production",
+	"CreatedBy": "Epic Games, Inc.",
+	"CreatedByURL": "http://epicgames.com",
+	"DocsURL": "",
+	"MarketplaceURL": "",
+	"SupportURL": "",
+	"CanContainContent": true,
+	"EnabledByDefault": false,
+	"IsBetaVersion": false,
+	"IsExperimentalVersion": true,
+	"Installed": false,
+	"Hidden": true,
+	"SupportedPrograms": [ "LiveLinkHub" ],
+	"Modules": [
+		{
+			"Name": "CaptureUtils",
+			"Type": "EditorAndProgram",
+			"ProgramAllowList": [ "LiveLinkHub" ],
+			"LoadingPhase": "PreDefault"
+		},
+		{
+			"Name": "CaptureSourceFramework",
+			"Type": "EditorAndProgram",
+			"ProgramAllowList": [ "LiveLinkHub" ],
+			"LoadingPhase": "PreDefault"
+		},
+		{
+			"Name": "CaptureProtocolStack",
+			"Type": "EditorAndProgram",
+			"ProgramAllowList": [ "LiveLinkHub" ],
+			"LoadingPhase": "Default"
+		},
+		{
+			"Name": "DataIngestCore",
+			"Type": "EditorAndProgram",
+			"ProgramAllowList": [ "LiveLinkHub" ],
+			"LoadingPhase": "Default"
+		},
+		{
+			"Name": "LiveLinkHubCaptureMessaging",
+			"Type": "EditorAndProgram",
+			"ProgramAllowList": [ "LiveLinkHub" ],
+			"LoadingPhase": "Default"
+		},
+		{
+			"Name": "CaptureManagerTakeMetadata",
+			"Type": "EditorAndProgram",
+			"ProgramAllowList": [ "LiveLinkHub" ],
+			"LoadingPhase": "Default"
+		}		
+	],
+	"Plugins": [
+		{
+			"Name": "CaptureData",
+			"Enabled": true
+		}
+	]
+}
\ No newline at end of file
diff --git a/Engine/Plugins/VirtualProduction/CaptureManager/CaptureManagerCore/Source/CaptureProtocolStack/Private/Tests/Control/TestControlRequest.cpp b/Engine/Plugins/VirtualProduction/CaptureManager/CaptureManagerCore/Source/CaptureProtocolStack/Private/Tests/Control/TestControlRequest.cpp
new file mode 100644
--- /dev/null
+++ b/Engine/Plugins/VirtualProduction/CaptureManager/CaptureManagerCore/Source/CaptureProtocolStack/Private/Tests/Control/TestControlRequest.cpp
@@ -0,0 +1,187 @@
+// Copyright Epic Games, Inc. All Rights Reserved.
+
+#include "Control/Messages/ControlRequest.h"
+
+#include "Control/Messages/ControlJsonUtilities.h"
+#include "Control/Messages/Constants.h"
+
+#include "Misc/AutomationTest.h"
+
+#if WITH_DEV_AUTOMATION_TESTS
+
+namespace UE::CaptureManager
+{
+
+IMPLEMENT_SIMPLE_AUTOMATION_TEST(TGetServerInformationTest, "Plugin.CaptureProtocolStack.Control.ControlRequest.GetServerInformation.Success", EAutomationTestFlags_ApplicationContextMask | EAutomationTestFlags::SmokeFilter)
+bool TGetServerInformationTest::RunTest(const FString& InParameters)
+{
+	FGetServerInformationRequest Request;
+
+	TestEqual(TEXT("GetAddresPath"), Request.GetAddressPath(), CPS::AddressPaths::GGetServerInformation);
+	TestFalse(TEXT("GetBody"), Request.GetBody().IsValid()); // Request doesn't have a Body
+
+	return true;
+}
+
+IMPLEMENT_SIMPLE_AUTOMATION_TEST(TKeepAliveRequestTest, "Plugin.CaptureProtocolStack.Control.ControlRequest.KeepAlive.Success", EAutomationTestFlags_ApplicationContextMask | EAutomationTestFlags::SmokeFilter)
+bool TKeepAliveRequestTest::RunTest(const FString& InParameters)
+{
+	FKeepAliveRequest Request;
+
+	TestEqual(TEXT("GetAddresPath"), Request.GetAddressPath(), CPS::AddressPaths::GKeepAlive);
+	TestFalse(TEXT("GetBody"), Request.GetBody().IsValid()); // Request doesn't have a Body
+
+	return true;
+}
+
+IMPLEMENT_SIMPLE_AUTOMATION_TEST(TStartSessionRequestTest, "Plugin.CaptureProtocolStack.Control.ControlRequest.StartSession.Success", EAutomationTestFlags_ApplicationContextMask | EAutomationTestFlags::SmokeFilter)
+bool TStartSessionRequestTest::RunTest(const FString& InParameters)
+{
+	FStartSessionRequest Request;
+
+	TestEqual(TEXT("GetAddresPath"), Request.GetAddressPath(), CPS::AddressPaths::GStartSession);
+	TestFalse(TEXT("GetBody"), Request.GetBody().IsValid()); // Request doesn't have a Body
+
+	return true;
+}
+
+IMPLEMENT_SIMPLE_AUTOMATION_TEST(TStopSessionRequestTest, "Plugin.CaptureProtocolStack.Control.ControlRequest.StopSession.Success", EAutomationTestFlags_ApplicationContextMask | EAutomationTestFlags::SmokeFilter)
+bool TStopSessionRequestTest::RunTest(const FString& InParameters)
+{
+	FStopSessionRequest Request;
+
+	TestEqual(TEXT("GetAddresPath"), Request.GetAddressPath(), CPS::AddressPaths::GStopSession);
+	TestFalse(TEXT("GetBody"), Request.GetBody().IsValid()); // Request doesn't have a Body
+
+	return true;
+}
+
+IMPLEMENT_SIMPLE_AUTOMATION_TEST(TSubscribeRequestTest, "Plugin.CaptureProtocolStack.Control.ControlRequest.Subscribe.Success", EAutomationTestFlags_ApplicationContextMask | EAutomationTestFlags::SmokeFilter)
+bool TSubscribeRequestTest::RunTest(const FString& InParameters)
+{
+	FSubscribeRequest Request;
+
+	TestEqual(TEXT("GetAddresPath"), Request.GetAddressPath(), CPS::AddressPaths::GSubscribe);
+	TestFalse(TEXT("GetBody"), Request.GetBody().IsValid()); // Request doesn't have a Body
+
+	return true;
+}
+
+IMPLEMENT_SIMPLE_AUTOMATION_TEST(TUnsubscribeRequestTest, "Plugin.CaptureProtocolStack.Control.ControlRequest.Unsubscribe.Success", EAutomationTestFlags_ApplicationContextMask | EAutomationTestFlags::SmokeFilter)
+bool TUnsubscribeRequestTest::RunTest(const FString& InParameters)
+{
+	FUnsubscribeRequest Request;
+
+	TestEqual(TEXT("GetAddresPath"), Request.GetAddressPath(), CPS::AddressPaths::GUnsubscribe);
+	TestFalse(TEXT("GetBody"), Request.GetBody().IsValid()); // Request doesn't have a Body
+
+	return true;
+}
+
+IMPLEMENT_SIMPLE_AUTOMATION_TEST(TGetStateRequestTest, "Plugin.CaptureProtocolStack.Control.ControlRequest.GetState.Success", EAutomationTestFlags_ApplicationContextMask | EAutomationTestFlags::SmokeFilter)
+bool TGetStateRequestTest::RunTest(const FString& InParameters)
+{
+	FGetStateRequest Request;
+
+	TestEqual(TEXT("GetAddresPath"), Request.GetAddressPath(), CPS::AddressPaths::GGetState);
+	TestFalse(TEXT("GetBody"), Request.GetBody().IsValid()); // Request doesn't have a Body
+
+	return true;
+}
+
+IMPLEMENT_SIMPLE_AUTOMATION_TEST(TStartRecordingTakeRequestTest, "Plugin.CaptureProtocolStack.Control.ControlRequest.StartRecordingTake.Success", EAutomationTestFlags_ApplicationContextMask | EAutomationTestFlags::SmokeFilter)
+bool TStartRecordingTakeRequestTest::RunTest(const FString& InParameters)
+{
+	FString String = TEXT("{\
+        \"slateName\": \"Slate\",\
+        \"takeNumber\": 0,\
+        \"subject\": \"Subject\",\
+        \"scenario\": \"Scenario\",\
+        \"tags\": [\"Tag1\", \"Tag2\", \"Tag3\"]\
+    }");
+
+	String.ConvertTabsToSpacesInline(2);
+	String.RemoveSpacesInline();
+
+	auto ConvertedString = StringCast<UTF8CHAR>(*String);
+
+	FString SlateName = TEXT("Slate");
+	FString Subject = TEXT("Subject");
+	FString Scenario = TEXT("Scenario");
+	TArray<FString> Tags = { TEXT("Tag1"), TEXT("Tag2"), TEXT("Tag3") };
+
+	FStartRecordingTakeRequest Request(SlateName, 0, Subject, Scenario, Tags);
+
+	TArray<uint8> Body;
+	TestTrue(TEXT("CreateDataFromJson"), FJsonUtility::CreateUTF8DataFromJson(Request.GetBody(), Body));
+
+	TestEqual(TEXT("GetAddresPath"), Request.GetAddressPath(), CPS::AddressPaths::GStartRecordingTake);
+
+	TArray<uint8> BodyExpected(reinterpret_cast<const uint8*>(ConvertedString.Get()), ConvertedString.Length());
+	TestEqual(TEXT("GetBody"), Body, BodyExpected);
+
+	return true;
+}
+
+IMPLEMENT_SIMPLE_AUTOMATION_TEST(TStopRecordingTakeRequestTest, "Plugin.CaptureProtocolStack.Control.ControlRequest.StopRecordingTake.Success", EAutomationTestFlags_ApplicationContextMask | EAutomationTestFlags::SmokeFilter)
+bool TStopRecordingTakeRequestTest::RunTest(const FString& InParameters)
+{
+	FStopRecordingTakeRequest Request;
+
+	TestEqual(TEXT("GetAddresPath"), Request.GetAddressPath(), CPS::AddressPaths::GStopRecordingTake);
+	TestFalse(TEXT("GetBody"), Request.GetBody().IsValid()); // Request doesn't have a Body
+
+	return true;
+}
+
+IMPLEMENT_SIMPLE_AUTOMATION_TEST(TAbortRecordingTakeRequestTest, "Plugin.CaptureProtocolStack.Control.ControlRequest.AbortRecordingTake.Success", EAutomationTestFlags_ApplicationContextMask | EAutomationTestFlags::SmokeFilter)
+bool TAbortRecordingTakeRequestTest::RunTest(const FString& InParameters)
+{
+	FAbortRecordingTakeRequest Request;
+
+	TestEqual(TEXT("GetAddresPath"), Request.GetAddressPath(), CPS::AddressPaths::GAbortRecordingTake);
+	TestFalse(TEXT("GetBody"), Request.GetBody().IsValid()); // Request doesn't have a Body
+
+	return true;
+}
+
+IMPLEMENT_SIMPLE_AUTOMATION_TEST(TGetTakeListRequestTest, "Plugin.CaptureProtocolStack.Control.ControlRequest.GetTakeList.Success", EAutomationTestFlags_ApplicationContextMask | EAutomationTestFlags::SmokeFilter)
+bool TGetTakeListRequestTest::RunTest(const FString& InParameters)
+{
+	FGetTakeListRequest Request;
+
+	TestEqual(TEXT("GetAddresPath"), Request.GetAddressPath(), CPS::AddressPaths::GGetTakeList);
+	TestFalse(TEXT("GetBody"), Request.GetBody().IsValid()); // Request doesn't have a Body
+
+	return true;
+}
+
+IMPLEMENT_SIMPLE_AUTOMATION_TEST(TGetTakeMetadataRequestTest, "Plugin.CaptureProtocolStack.Control.ControlRequest.GetTakeMetadata.Success", EAutomationTestFlags_ApplicationContextMask | EAutomationTestFlags::SmokeFilter)
+bool TGetTakeMetadataRequestTest::RunTest(const FString& InParameters)
+{
+	FString String = TEXT("{\
+        \"names\": [\"TakeName1\", \"TakeName2\", \"TakeName3\"]\
+    }");
+
+	String.ConvertTabsToSpacesInline(2);
+	String.RemoveSpacesInline();
+
+	auto ConvertedString = StringCast<UTF8CHAR>(*String);
+
+	TArray<FString> Takes = { TEXT("TakeName1"), TEXT("TakeName2"), TEXT("TakeName3") };
+	FGetTakeMetadataRequest Request(Takes);
+
+	TArray<uint8> Body;
+	TestTrue(TEXT("CreateDataFromJson"), FJsonUtility::CreateUTF8DataFromJson(Request.GetBody(), Body));
+
+	TestEqual(TEXT("GetAddresPath"), Request.GetAddressPath(), CPS::AddressPaths::GGetTakeMetadata);
+
+	TArray<uint8> BodyExpected(reinterpret_cast<const uint8*>(ConvertedString.Get()), ConvertedString.Length());
+	TestEqual(TEXT("GetBody"), Body, BodyExpected);
+
+	return true;
+}
+
+}
+
+#endif // WITH_DEV_AUTOMATION_TESTS
\ No newline at end of file
