Fix missing plugin in incremental build

diff --git a/Engine/Plugins/VirtualProduction/CaptureManager/CaptureManagerApp/CaptureManagerApp.uplugin b/Engine/Plugins/VirtualProduction/CaptureManager/CaptureManagerApp/CaptureManagerApp.uplugin
--- a/Engine/Plugins/VirtualProduction/CaptureManager/CaptureManagerApp/CaptureManagerApp.uplugin
+++ b/Engine/Plugins/VirtualProduction/CaptureManager/CaptureManagerApp/CaptureManagerApp.uplugin
@@ -101,12 +101,6 @@
 			"Type": "EditorAndProgram",
 			"ProgramAllowList": [ "LiveLinkHub" ],
 			"LoadingPhase": "Default"
-		},
-		{
-			"Name": "ExampleLiveLinkDevices",
-			"Type": "EditorAndProgram",
-			"ProgramAllowList": [ "LiveLinkHub" ],
-			"LoadingPhase": "Default"
 		}
 	],
 	"Plugins": [
