Implement python API using new LiveLink Device

diff --git a/Engine/Plugins/VirtualProduction/CaptureManager/CaptureManagerApp/Source/CaptureManagerEditor/CaptureManagerEditor.Build.cs b/Engine/Plugins/VirtualProduction/CaptureManager/CaptureManagerApp/Source/CaptureManagerEditor/CaptureManagerEditor.Build.cs
--- a/Engine/Plugins/VirtualProduction/CaptureManager/CaptureManagerApp/Source/CaptureManagerEditor/CaptureManagerEditor.Build.cs
+++ b/Engine/Plugins/VirtualProduction/CaptureManager/CaptureManagerApp/Source/CaptureManagerEditor/CaptureManagerEditor.Build.cs
@@ -18,7 +18,7 @@ public class CaptureManagerEditor : ModuleRules
 
 		PrivateDependencyModuleNames.AddRange(new string[]
 		{
-			"CaptureSourceFramework",
+			"GenericLiveLinkDevice",
 			"InputCore",
 			"LiveLinkDevice",
 			"LiveLinkCapabilities",
diff --git a/Engine/Plugins/VirtualProduction/CaptureManager/CaptureManagerApp/Source/CaptureManagerEditor/Private/IngestManagement/IngestJobExecutor.cpp b/Engine/Plugins/VirtualProduction/CaptureManager/CaptureManagerApp/Source/CaptureManagerEditor/Private/IngestManagement/IngestJobExecutor.cpp
--- a/Engine/Plugins/VirtualProduction/CaptureManager/CaptureManagerApp/Source/CaptureManagerEditor/Private/IngestManagement/IngestJobExecutor.cpp
+++ b/Engine/Plugins/VirtualProduction/CaptureManager/CaptureManagerApp/Source/CaptureManagerEditor/Private/IngestManagement/IngestJobExecutor.cpp
@@ -201,7 +201,7 @@ UIngestCapability_ProcessHandle* FIngestJobExecutor::RunIngest(FIngestJob& InJob
 	const FIngestJob::FSettings Settings = InJob.GetSettings();
 
 	UIngestCapability_Options* IngestOptions = NewObject<UIngestCapability_Options>();
-	IngestOptions->Location = Settings.LocalWorkFolder;
+	IngestOptions->WorkingDirectory = Settings.LocalWorkFolder;
 	FIngestCapability_VideoOptions VideoOptions;
 	VideoOptions.FileNamePrefix = Settings.VideoSettings.FileNamePrefix;
 	VideoOptions.Format = Settings.VideoSettings.Format;
diff --git a/Engine/Plugins/VirtualProduction/CaptureManager/CaptureManagerApp/Source/CaptureManagerEditor/Private/Scripting/ULiveLinkHubArchiveCaptureDevice.cpp b/Engine/Plugins/VirtualProduction/CaptureManager/CaptureManagerApp/Source/CaptureManagerEditor/Private/Scripting/ULiveLinkHubArchiveCaptureDevice.cpp
deleted file mode 100644
diff --git a/Engine/Plugins/VirtualProduction/CaptureManager/CaptureManagerApp/Source/CaptureManagerEditor/Private/Scripting/ULiveLinkHubArchiveCaptureDevice.h b/Engine/Plugins/VirtualProduction/CaptureManager/CaptureManagerApp/Source/CaptureManagerEditor/Private/Scripting/ULiveLinkHubArchiveCaptureDevice.h
deleted file mode 100644
diff --git a/Engine/Plugins/VirtualProduction/CaptureManager/CaptureManagerApp/Source/CaptureManagerEditor/Private/Scripting/ULiveLinkHubCaptureDevice.cpp b/Engine/Plugins/VirtualProduction/CaptureManager/CaptureManagerApp/Source/CaptureManagerEditor/Private/Scripting/ULiveLinkHubCaptureDevice.cpp
--- a/Engine/Plugins/VirtualProduction/CaptureManager/CaptureManagerApp/Source/CaptureManagerEditor/Private/Scripting/ULiveLinkHubCaptureDevice.cpp
+++ b/Engine/Plugins/VirtualProduction/CaptureManager/CaptureManagerApp/Source/CaptureManagerEditor/Private/Scripting/ULiveLinkHubCaptureDevice.cpp
@@ -2,9 +2,14 @@
 
 #include "ULiveLinkHubCaptureDevice.h"
 
-#include "CaptureSourceFrameworkModule.h"
+#include "LiveLinkDeviceCapability_Connection.h"
+#include "Ingest/LiveLinkDeviceCapability_Ingest.h"
 
-#include "Capabilities/IngestCapability.h"
+#include "GenericLiveLinkDeviceHelper.h"
+
+#include "Engine/Engine.h"
+
+#include "LiveLinkDeviceSubsystem.h"
 
 #define LOCTEXT_NAMESPACE "LiveLinkHubScriptedIngest"
 
@@ -20,7 +25,7 @@ struct ULiveLinkHubCaptureDevice::FImpl
 {
 	struct FSession
 	{
-		explicit FSession(UE::CaptureManager::FCaptureSourceId InCaptureSourceId);
+		explicit FSession(const TObjectPtr<ULiveLinkDevice> InDevice);
 
 		// This session ID is used to verify that any take objects we have returned to the user still belong to the 
 		// current start/stop session. This should prevent issues where take IDs have changed after a start/stop call 
@@ -29,200 +34,172 @@ struct ULiveLinkHubCaptureDevice::FImpl
 		// may now have a different take ID.
 		const FGuid Id;
 
-		const UE::CaptureManager::FCaptureSourceId CaptureSourceId;
+		FGuid CaptureDeviceId;
 	};
 
-	static UE::CaptureManager::FCaptureSourceManager& GetCaptureSourceManager();
-
 	FImpl();
 
-	FLiveLinkHubScriptStatus ImportTake(const FLiveLinkHubTakeMetadata& InTake, const FLiveLinkHubConversionSettings& InConversionSettings) const;
-	FLiveLinkHubScriptStatus FetchTakes(TArray<FLiveLinkHubTakeMetadata>& OutTakes) const;
+	TObjectPtr<UIngestCapability_ProcessResult> ImportTake(const FLiveLinkHubTakeMetadata& InTake, 
+														   const TObjectPtr<const UIngestCapability_Options> InConversionSettings) const;
+	TObjectPtr<UIngestCapability_ProcessResult> FetchTakes(TArray<FLiveLinkHubTakeMetadata>& OutTakes) const;
 
-	FLiveLinkHubScriptStatus Start(const int32 InTimeoutSeconds, TMap<FString, FPropertyValue> InParamValues, const FString& InFactoryId, const FString& InName);
-	FLiveLinkHubScriptStatus Stop();
+	TObjectPtr<UIngestCapability_ProcessResult> Start(const int32 InTimeoutSeconds);
+	TObjectPtr<UIngestCapability_ProcessResult> Stop();
 
-	bool StartCaptureSource(UE::CaptureManager::FCaptureSourceId InCaptureSourceId, int32 InTimeoutSeconds);
-	void StopCaptureSource(UE::CaptureManager::FCaptureSourceId InCaptureSourceId);
+	bool StartCaptureDevice(int32 InTimeoutSeconds);
+	void StopCaptureDevice();
+	void RemoveDevice();
 
 	TOptional<FSession> Session;
-	UE::CaptureManager::FCaptureSourceManager& CaptureSourceManager;
+
+	TObjectPtr<ULiveLinkDevice> BackgroundDevice;
 };
 
-ULiveLinkHubCaptureDevice::FImpl::FSession::FSession(const UE::CaptureManager::FCaptureSourceId InCaptureSourceId) :
-	Id(FGuid::NewGuid()),
-	CaptureSourceId(InCaptureSourceId)
+ULiveLinkHubCaptureDevice::FImpl::FSession::FSession(const TObjectPtr<ULiveLinkDevice> InDevice) :
+	Id(FGuid::NewGuid())
 {
-}
+	ULiveLinkDeviceSubsystem* Subsystem = GEngine->GetEngineSubsystem<ULiveLinkDeviceSubsystem>();
 
-UE::CaptureManager::FCaptureSourceManager& ULiveLinkHubCaptureDevice::FImpl::GetCaptureSourceManager()
-{
-	using namespace  UE::CaptureManager;
+	const FGuid* MaybeDevice = Subsystem->GetDeviceMap().FindKey(InDevice);
+	check(MaybeDevice);
 
-	FCaptureSourceFrameworkModule& CaptureSourceModule = FModuleManager::LoadModuleChecked<FCaptureSourceFrameworkModule>("CaptureSourceFramework");
-	return CaptureSourceModule.GetCaptureSourceManager();
+	CaptureDeviceId = *MaybeDevice;
 }
 
-ULiveLinkHubCaptureDevice::FImpl::FImpl() :
-	CaptureSourceManager(GetCaptureSourceManager())
+ULiveLinkHubCaptureDevice::FImpl::FImpl()
 {
 }
 
-FLiveLinkHubScriptStatus ULiveLinkHubCaptureDevice::FImpl::Start(const int32 InTimeoutSeconds, TMap<FString, FPropertyValue> InParamValues, const FString& InFactoryId, const FString& InName)
+TObjectPtr<UIngestCapability_ProcessResult> ULiveLinkHubCaptureDevice::FImpl::Start(const int32 InTimeoutSeconds)
 {
-	using namespace  UE::CaptureManager;
+	using namespace UE::CaptureManager;
 
 	if (Session)
 	{
-		return FLiveLinkHubScriptStatus::Success();
+		return UIngestCapability_ProcessResult::Success();
 	}
 
-	FCaptureSourceManager::FCaptureSourceResult CaptureSourceId = CaptureSourceManager.CreateCaptureSource(InFactoryId, InName, MoveTemp(InParamValues));
-
-	if (CaptureSourceId.HasError())
-	{
-		return FLiveLinkHubScriptStatus::Error(FText::Format(FailedToCreateCaptureSourceFormatText, FText::FromString(CaptureSourceId.GetError().GetMessage())));
-	}
-
-	const bool bIsStarted = StartCaptureSource(CaptureSourceId.GetValue(), InTimeoutSeconds);
+	const bool bIsStarted = StartCaptureDevice(InTimeoutSeconds);
 
 	if (bIsStarted)
 	{
-		Session = FSession(CaptureSourceId.StealValue());
-		return FLiveLinkHubScriptStatus::Success();
-	}
-	else
-	{
-		return FLiveLinkHubScriptStatus::Error(StartTimedOutText);
+		Session = FSession(BackgroundDevice);
+		return UIngestCapability_ProcessResult::Success();
 	}
+
+	return UIngestCapability_ProcessResult::Error(StartTimedOutText);
 }
 
-bool ULiveLinkHubCaptureDevice::FImpl::StartCaptureSource(const UE::CaptureManager::FCaptureSourceId InCaptureSourceId, const int32 InTimeoutSeconds)
+bool ULiveLinkHubCaptureDevice::FImpl::StartCaptureDevice(const int32 InTimeoutSeconds)
 {
-	using namespace UE;
 	using namespace UE::CaptureManager;
 
-	FCaptureSourceHandle CaptureSourceHandle = CaptureSourceManager.GetCaptureSource(InCaptureSourceId);
 	FEventRef ReachableEvent;
 
-	TManagedDelegate<TSharedPtr<const FCaptureEvent>> OnReachableEvent(
-		[&ReachableEvent](const TSharedPtr<const FCaptureEvent> CaptureEvent)
-		{
-			if (CaptureEvent)
-			{
-				TSharedPtr<const FReachableEvent> Event = StaticCastSharedPtr<const FReachableEvent>(CaptureEvent);
-				if (Event->bIsReachable)
-				{
-					ReachableEvent->Trigger();
-				}
-			}
-		},
-		EDelegateExecutionThread::InternalThread
-	);
-
-	CaptureSourceHandle.SubscribeToEvent(FReachableEvent::Name, MoveTemp(OnReachableEvent));
-
-	ON_SCOPE_EXIT
+	UConnectionDelegate* ConnectionDelegate = ILiveLinkDeviceCapability_Connection::Execute_GetConnectionDelegate(BackgroundDevice);
+	FDelegateHandle DelegateHandle = 
+		ConnectionDelegate->ConnectionChanged.AddLambda([&ReachableEvent](ELiveLinkDeviceConnectionStatus InEvent)
 	{
-		CaptureSourceHandle.UnsubscribeAll();
-	};
+		if (InEvent == ELiveLinkDeviceConnectionStatus::Connected)
+		{
+			ReachableEvent->Trigger();
+		}
+	});
 
 	AsyncTask(
 		ENamedThreads::AnyThread,
-		[&CaptureSourceHandle]()
+		[this]()
 		{
-			CaptureSourceHandle.Start();
+			ILiveLinkDeviceCapability_Connection::Execute_Connect(BackgroundDevice);
 		}
 	);
 
-	return ReachableEvent->Wait(FTimespan::FromSeconds(InTimeoutSeconds));
+	bool bResult = ReachableEvent->Wait(FTimespan::FromSeconds(InTimeoutSeconds));
+
+	ConnectionDelegate->ConnectionChanged.Remove(DelegateHandle);
+
+	return bResult;
 }
 
-FLiveLinkHubScriptStatus ULiveLinkHubCaptureDevice::FImpl::Stop()
+TObjectPtr<UIngestCapability_ProcessResult> ULiveLinkHubCaptureDevice::FImpl::Stop()
 {
-	using namespace  UE::CaptureManager;
+	using namespace UE::CaptureManager;
 
 	if (!Session)
 	{
-		return FLiveLinkHubScriptStatus::Success();
+		return UIngestCapability_ProcessResult::Success();
 	}
 
-	StopCaptureSource(Session->CaptureSourceId);
-	CaptureSourceManager.RemoveCaptureSource(Session->CaptureSourceId);
+	StopCaptureDevice();
+	RemoveDevice();
 	Session.Reset();
 
-	return FLiveLinkHubScriptStatus::Success();
+	return UIngestCapability_ProcessResult::Success();
 }
 
-void ULiveLinkHubCaptureDevice::FImpl::StopCaptureSource(const UE::CaptureManager::FCaptureSourceId InCaptureSourceId)
+void ULiveLinkHubCaptureDevice::FImpl::StopCaptureDevice()
 {
-	using namespace  UE::CaptureManager;
+	using namespace UE::CaptureManager;
 
-	FCaptureSourceHandle CaptureSourceHandle = CaptureSourceManager.GetCaptureSource(InCaptureSourceId);
 	FEventRef UnreachableEvent;
 
-	TManagedDelegate<TSharedPtr<const FCaptureEvent>> OnReachableEvent(
-		[&UnreachableEvent](const TSharedPtr<const FCaptureEvent> CaptureEvent)
+	auto OnReachableEvent = [&UnreachableEvent](ELiveLinkDeviceConnectionStatus InEvent)
 		{
-			if (CaptureEvent)
+			if (InEvent == ELiveLinkDeviceConnectionStatus::Disconnected)
 			{
-				TSharedPtr<const FReachableEvent> Event = StaticCastSharedPtr<const FReachableEvent>(CaptureEvent);
-				if (!Event->bIsReachable)
-				{
-					UnreachableEvent->Trigger();
-					return;
-				}
+				UnreachableEvent->Trigger();
+				return;
 			}
-		},
-		EDelegateExecutionThread::InternalThread
-	);
+		};
 
-	CaptureSourceHandle.SubscribeToEvent(FReachableEvent::Name, MoveTemp(OnReachableEvent));
-
-	ON_SCOPE_EXIT
-	{
-		CaptureSourceHandle.UnsubscribeAll();
-	};
+	UConnectionDelegate* ConnectionDelegate = ILiveLinkDeviceCapability_Connection::Execute_GetConnectionDelegate(BackgroundDevice);
+	FDelegateHandle DelegateHandle = ConnectionDelegate->ConnectionChanged.AddLambda(MoveTemp(OnReachableEvent));
 
 	AsyncTask(
 		ENamedThreads::AnyThread,
-		[&CaptureSourceHandle]()
+		[this]()
 		{
-			CaptureSourceHandle.Stop();
+			ILiveLinkDeviceCapability_Connection::Execute_Disconnect(BackgroundDevice);
 		}
 	);
 
 	UnreachableEvent->Wait();
+
+	ConnectionDelegate->ConnectionChanged.Remove(DelegateHandle);
 }
 
-FLiveLinkHubScriptStatus ULiveLinkHubCaptureDevice::FImpl::ImportTake(const FLiveLinkHubTakeMetadata& InTake, const FLiveLinkHubConversionSettings& InConversionSettings) const
+void ULiveLinkHubCaptureDevice::FImpl::RemoveDevice()
 {
-	using namespace  UE::CaptureManager;
+	ULiveLinkDeviceSubsystem* Subsystem = GEngine->GetEngineSubsystem<ULiveLinkDeviceSubsystem>();
+	Subsystem->RemoveDevice(BackgroundDevice);
+}
+
+TObjectPtr<UIngestCapability_ProcessResult> ULiveLinkHubCaptureDevice::FImpl::ImportTake(const FLiveLinkHubTakeMetadata& InTake, 
+																						 const TObjectPtr<const UIngestCapability_Options> InConversionSettings) const
+{
+	using namespace UE::CaptureManager;
 
 	if (!Session)
 	{
-		return FLiveLinkHubScriptStatus::Error(NoSessionText);
+		return UIngestCapability_ProcessResult::Error(NoSessionText);
 	}
 
 	if (InTake.SessionId != Session->Id)
 	{
-		return FLiveLinkHubScriptStatus::Error(InvalidTakeSessionIdText);
+		return UIngestCapability_ProcessResult::Error(InvalidTakeSessionIdText);
 	}
 
-	FCaptureSourceHandle CaptureSourceHandle = CaptureSourceManager.GetCaptureSource(Session->CaptureSourceId);
-	FIngestCapability* IngestCapability = CaptureSourceHandle.GetCapability<FIngestCapability>(FIngestCapability::Name);
-
-	if (!IngestCapability)
+	if (!BackgroundDevice->Implements<ULiveLinkDeviceCapability_Ingest>())
 	{
-		return FLiveLinkHubScriptStatus::Error(IngestNotSupportedText);
+		return UIngestCapability_ProcessResult::Error(IngestNotSupportedText);
 	}
 
 	FEventRef ImportCompleteEvent;
 	bool bImportSuccess = false;
 	FText ErrorMessage;
 
-	TManagedDelegate<FTakeId, TValueOrError<void, FIngestCapabilityError>> OnImportFinished(
-		[&ImportCompleteEvent, &bImportSuccess, &ErrorMessage](FTakeId InTakeId, TValueOrError<void, FIngestCapabilityError> InImportResult)
+	TManagedDelegate<FTakeId, TValueOrError<void, FIngestCapability_Error>> OnIngestFinished(
+		[&ImportCompleteEvent, &bImportSuccess, &ErrorMessage](FTakeId InTakeId, TValueOrError<void, FIngestCapability_Error> InImportResult)
 		{
 			if (InImportResult.HasValue())
 			{
@@ -238,38 +215,24 @@ FLiveLinkHubScriptStatus ULiveLinkHubCaptureDevice::FImpl::ImportTake(const FLiv
 		EDelegateExecutionThread::InternalThread
 	);
 
-	FIngestOptions::FAudio AudioOptions =
-	{
-		.FileNamePrefix = InConversionSettings.AudioFileNamePrefix,
-		.Format = InConversionSettings.AudioFormat
-	};
+	UIngestCapability_ProcessHandle* ProcessHandle =
+		ILiveLinkDeviceCapability_Ingest::Execute_PrepareIngestProcess(BackgroundDevice, InTake.TakeId);
+	check(ProcessHandle);
 
-	FIngestOptions::FVideo VideoOptions = 
-	{
-		.FileNamePrefix = InConversionSettings.VideoFileNamePrefix,
-		.Format = InConversionSettings.VideoFormat,
-		.PixelFormat = UE::CaptureManager::EImagePixelFormat::Default
-	};
-
-	FIngestOptions Options =
-	{
-		.Location = InConversionSettings.WorkingDirectory,
-		.Video = VideoOptions,
-		.Audio = AudioOptions
-	};
+	ProcessHandle->ProcessFinishedReporter = MoveTemp(OnIngestFinished);
+	ProcessHandle->Run(InConversionSettings);
 
-	IngestCapability->ImportTakeAsync(InTake.TakeId, MoveTemp(Options), OnImportFinished, TManagedDelegate<FTakeId, double>());
 	ImportCompleteEvent->Wait();
 
 	if (bImportSuccess)
 	{
-		return FLiveLinkHubScriptStatus::Success();
+		return UIngestCapability_ProcessResult::Success();
 	}
 
-	return FLiveLinkHubScriptStatus::Error(ErrorMessage);
+	return UIngestCapability_ProcessResult::Error(ErrorMessage);
 }
 
-FLiveLinkHubScriptStatus ULiveLinkHubCaptureDevice::FImpl::FetchTakes(TArray<FLiveLinkHubTakeMetadata>& OutTakes) const
+TObjectPtr<UIngestCapability_ProcessResult> ULiveLinkHubCaptureDevice::FImpl::FetchTakes(TArray<FLiveLinkHubTakeMetadata>& OutTakes) const
 {
 	using namespace  UE::CaptureManager;
 
@@ -281,40 +244,30 @@ FLiveLinkHubScriptStatus ULiveLinkHubCaptureDevice::FImpl::FetchTakes(TArray<FLi
 
 	if (!Session)
 	{
-		return FLiveLinkHubScriptStatus::Error(NoSessionText);
+		return UIngestCapability_ProcessResult::Error(NoSessionText);
 	}
 
-	FCaptureSourceHandle CaptureSourceHandle = CaptureSourceManager.GetCaptureSource(Session->CaptureSourceId);
-	FIngestCapability* IngestCapability = CaptureSourceHandle.GetCapability<FIngestCapability>(FIngestCapability::Name);
-
-	if (!IngestCapability)
+	if (!BackgroundDevice->Implements<ULiveLinkDeviceCapability_Ingest>())
 	{
-		return FLiveLinkHubScriptStatus::Error(IngestNotSupportedText);
+		return UIngestCapability_ProcessResult::Error(IngestNotSupportedText);
 	}
 
 	FEventRef TakeListUpdatedEvent;
 	TArray<FLiveLinkHubTakeMetadata> Takes;
 
 	TManagedDelegate<TArray<FTakeId>> OnTakesListUpdated(
-		[SessionId = Session->Id, IngestCapability, &Takes, &TakeListUpdatedEvent](const TArray<FTakeId>& InTakeIds)
+		[SessionId = Session->Id, Device = BackgroundDevice, &Takes, &TakeListUpdatedEvent](const TArray<FTakeId>& InTakeIds)
 		{
 			for (const FTakeId& TakeId : InTakeIds)
 			{
-				TOptional<FTakeMetadata> TakeMetadataOpt = IngestCapability->GetTakeMetadata(TakeId);
-
-				if (TakeMetadataOpt.IsSet())
-				{
-					const FTakeMetadata TakeMetadata = MoveTemp(TakeMetadataOpt.GetValue());
+				UIngestCapability_TakeInformation* Metadata =  ILiveLinkDeviceCapability_Ingest::Execute_GetTakeInformation(Device, TakeId);
 
-					FLiveLinkHubTakeMetadata Take;
-					Take.SlateName = TakeMetadata.Slate;
-					Take.TakeNumber = TakeMetadata.TakeNumber;
-					Take.TakeId = TakeId;
-					Take.SessionId = SessionId;
-					Take.DateTime = TakeMetadata.DateTime.Get(FDateTime()).ToIso8601();
+				FLiveLinkHubTakeMetadata Take;
+				Take.TakeId = TakeId;
+				Take.SessionId = SessionId;
+				Take.Metadata = Metadata;
 
-					Takes.Emplace(MoveTemp(Take));
-				}
+				Takes.Emplace(MoveTemp(Take));
 			}
 
 			TakeListUpdatedEvent->Trigger();
@@ -322,11 +275,44 @@ FLiveLinkHubScriptStatus ULiveLinkHubCaptureDevice::FImpl::FetchTakes(TArray<FLi
 		EDelegateExecutionThread::InternalThread
 	);
 
-	IngestCapability->UpdateTakeListAsync(MoveTemp(OnTakesListUpdated));
+	UIngestCapability_UpdateTakeListCallback* UpdateTakeListCallback =
+		ILiveLinkDeviceCapability_Ingest::Execute_CreateUpdateTakeListCallback(BackgroundDevice);
+	check(UpdateTakeListCallback);
+
+	UpdateTakeListCallback->Callback = MoveTemp(OnTakesListUpdated);
+
+	ILiveLinkDeviceCapability_Ingest::Execute_UpdateTakeListAsync(BackgroundDevice);
 	TakeListUpdatedEvent->Wait();
 
 	OutTakes = MoveTemp(Takes);
-	return FLiveLinkHubScriptStatus::Success();
+
+	return UIngestCapability_ProcessResult::Success();
+}
+
+ULiveLinkHubCaptureDevice* ULiveLinkHubCaptureDeviceFactory::CreateArchiveDevice(FString InName, FString InStoragePath)
+{
+	UGenericLiveLinkDeviceHelper* Helper = NewObject<UGenericLiveLinkDeviceHelper>(this);
+
+	TObjectPtr<ULiveLinkDevice> BackgroundDevice = Helper->CreateArchiveDevice(MoveTemp(InName), MoveTemp(InStoragePath));
+
+	TObjectPtr<ULiveLinkHubCaptureDevice> Device = NewObject<ULiveLinkHubCaptureDevice>(this);
+	Device->Name = MoveTemp(InName);
+	Device->ImplPtr->BackgroundDevice = MoveTemp(BackgroundDevice);
+
+	return Device;
+}
+
+ULiveLinkHubCaptureDevice* ULiveLinkHubCaptureDeviceFactory::CreateNetworkDevice(FString InName, FString InIPAddress, int32 InPort)
+{
+	UGenericLiveLinkDeviceHelper* Helper = NewObject<UGenericLiveLinkDeviceHelper>(this);
+
+	TObjectPtr<ULiveLinkDevice> BackgroundDevice = Helper->CreateNetworkDevice(InName, MoveTemp(InIPAddress), InPort);
+
+	TObjectPtr<ULiveLinkHubCaptureDevice> Device = NewObject<ULiveLinkHubCaptureDevice>(this);
+	Device->Name = MoveTemp(InName);
+	Device->ImplPtr->BackgroundDevice = MoveTemp(BackgroundDevice);
+
+	return Device;
 }
 
 ULiveLinkHubCaptureDevice::ULiveLinkHubCaptureDevice() :
@@ -336,26 +322,26 @@ ULiveLinkHubCaptureDevice::ULiveLinkHubCaptureDevice() :
 
 ULiveLinkHubCaptureDevice::~ULiveLinkHubCaptureDevice() = default;
 
-FLiveLinkHubScriptStatus ULiveLinkHubCaptureDevice::Start(const int32 InTimeoutSeconds)
+UIngestCapability_ProcessResult* ULiveLinkHubCaptureDevice::Start(const int32 InTimeoutSeconds)
 {
-	TMap<FString, FPropertyValue> ParamValues = GetCreationParameters();
-	return ImplPtr->Start(InTimeoutSeconds, MoveTemp(ParamValues), FactoryId, Name);
+	return ImplPtr->Start(InTimeoutSeconds);
 }
 
-FLiveLinkHubScriptStatus ULiveLinkHubCaptureDevice::Stop()
+UIngestCapability_ProcessResult* ULiveLinkHubCaptureDevice::Stop()
 {
 	return ImplPtr->Stop();
 }
 
-FLiveLinkHubScriptStatus ULiveLinkHubCaptureDevice::ImportTake(const FLiveLinkHubTakeMetadata& InTake, const FLiveLinkHubConversionSettings& InConversionSettings) const
+UIngestCapability_ProcessResult* ULiveLinkHubCaptureDevice::ImportTake(const FLiveLinkHubTakeMetadata& InTake, const UIngestCapability_Options* InConversionSettings) const
 {
-	return ImplPtr->ImportTake(InTake, InConversionSettings);
+	TObjectPtr<const UIngestCapability_Options> ConversionSettings = InConversionSettings;
+	return ImplPtr->ImportTake(InTake, MoveTemp(ConversionSettings));
 }
 
 FLiveLinkHubFetchTakesResult ULiveLinkHubCaptureDevice::FetchTakes() const
 {
 	TArray<FLiveLinkHubTakeMetadata> Takes;
-	FLiveLinkHubScriptStatus Status = ImplPtr->FetchTakes(Takes);
+	TObjectPtr<UIngestCapability_ProcessResult> Status = ImplPtr->FetchTakes(Takes);
 
 	FLiveLinkHubFetchTakesResult Result;
 	Result.Status = MoveTemp(Status);
diff --git a/Engine/Plugins/VirtualProduction/CaptureManager/CaptureManagerApp/Source/CaptureManagerEditor/Private/Scripting/ULiveLinkHubCaptureDevice.h b/Engine/Plugins/VirtualProduction/CaptureManager/CaptureManagerApp/Source/CaptureManagerEditor/Private/Scripting/ULiveLinkHubCaptureDevice.h
--- a/Engine/Plugins/VirtualProduction/CaptureManager/CaptureManagerApp/Source/CaptureManagerEditor/Private/Scripting/ULiveLinkHubCaptureDevice.h
+++ b/Engine/Plugins/VirtualProduction/CaptureManager/CaptureManagerApp/Source/CaptureManagerEditor/Private/Scripting/ULiveLinkHubCaptureDevice.h
@@ -2,9 +2,14 @@
 
 #pragma once
 
-#include "CaptureSourceCapabilityProperty.h"
 #include "LiveLinkHubScriptStatus.h"
 
+#include "Ingest/IngestCapability_TakeInformation.h"
+#include "Ingest/IngestCapability_ProcessHandle.h"
+#include "Ingest/IngestCapability_Options.h"
+
+#include "LiveLinkDevice.h"
+
 #include "ULiveLinkHubCaptureDevice.generated.h"
 
 USTRUCT(BlueprintType)
@@ -12,14 +17,8 @@ struct FLiveLinkHubTakeMetadata
 {
 	GENERATED_BODY()
 
-	UPROPERTY(EditAnywhere, BlueprintReadWrite, Category = Metadata)
-	FString SlateName;
-
-	UPROPERTY(EditAnywhere, BlueprintReadWrite, Category = Metadata)
-	int32 TakeNumber = 1;
-
-	UPROPERTY(EditAnywhere, BlueprintReadWrite, Category = Metadata)
-	FString DateTime;
+	UPROPERTY(VisibleAnywhere, BlueprintReadOnly, Category = Metadata)
+	TObjectPtr<UIngestCapability_TakeInformation> Metadata;
 
 	UPROPERTY()
 	int32 TakeId = 0;
@@ -35,38 +34,27 @@ struct FLiveLinkHubFetchTakesResult
 	GENERATED_BODY()
 
 	UPROPERTY(VisibleAnywhere, BlueprintReadOnly, Category = Result)
-	FLiveLinkHubScriptStatus Status;
+	TObjectPtr<UIngestCapability_ProcessResult> Status;
 
 	UPROPERTY(VisibleAnywhere, BlueprintReadOnly, Category = Result)
 	TArray<FLiveLinkHubTakeMetadata> Takes;
 };
 
-USTRUCT(BlueprintType)
-struct FLiveLinkHubConversionSettings
+UCLASS()
+class ULiveLinkHubCaptureDeviceFactory : public UObject
 {
 	GENERATED_BODY()
 
-	UPROPERTY(EditAnywhere, BlueprintReadWrite, Category = General)
-	FString WorkingDirectory;
-
-	UPROPERTY(EditAnywhere, BlueprintReadWrite, Category = Audio)
-	FString AudioFormat;
-
-	UPROPERTY(EditAnywhere, BlueprintReadWrite, Category = Audio)
-	FString AudioFileNamePrefix;
-
-	UPROPERTY(EditAnywhere, BlueprintReadWrite, Category = Video)
-	FString VideoFormat;
+public:
 
-	UPROPERTY(EditAnywhere, BlueprintReadWrite, Category = Video)
-	FString VideoFileNamePrefix;
+	UFUNCTION(BlueprintCallable, Category = "Ingest | Creator")
+	ULiveLinkHubCaptureDevice* CreateArchiveDevice(FString InName, FString InStoragePath);
 
-	// TODO: Make enum visible to python, or map python enum to c++ enum
-	// UPROPERTY(EditAnywhere, BlueprintReadWrite, Category = Video)
-	// int32 PixelFormat = 0;
+	UFUNCTION(BlueprintCallable, Category = "Ingest | Creator")
+	ULiveLinkHubCaptureDevice* CreateNetworkDevice(FString InName, FString InIPAddress, int32 InPort);
 };
 
-UCLASS(Abstract)
+UCLASS()
 class ULiveLinkHubCaptureDevice : public UObject
 {
 	GENERATED_BODY()
@@ -75,28 +63,26 @@ public:
 	ULiveLinkHubCaptureDevice();
 	~ULiveLinkHubCaptureDevice();
 
-	UFUNCTION(BlueprintCallable, Category = Control)
-	FLiveLinkHubScriptStatus Start(int32 InTimeoutSeconds);
+	UFUNCTION(BlueprintCallable, Category = Ingest)
+	UIngestCapability_ProcessResult* Start(int32 InTimeoutSeconds);
 
-	UFUNCTION(BlueprintCallable, Category = Control)
-	FLiveLinkHubScriptStatus Stop();
+	UFUNCTION(BlueprintCallable, Category = Ingest)
+	UIngestCapability_ProcessResult* Stop();
 
-	UFUNCTION(BlueprintCallable, Category = Control)
-	FLiveLinkHubScriptStatus ImportTake(const FLiveLinkHubTakeMetadata& InTake, const FLiveLinkHubConversionSettings& InConversionSettings) const;
+	UFUNCTION(BlueprintCallable, Category = Ingest)
+	UIngestCapability_ProcessResult* ImportTake(const FLiveLinkHubTakeMetadata& InTake, const UIngestCapability_Options* InConversionSettings) const;
 
-	UFUNCTION(BlueprintCallable, Category = Control)
+	UFUNCTION(BlueprintCallable, Category = Ingest)
 	FLiveLinkHubFetchTakesResult FetchTakes() const;
 
 	UPROPERTY(EditAnywhere, BlueprintReadWrite, Category = Info)
 	FString Name;
 
-	UPROPERTY(EditAnywhere, BlueprintReadWrite, Category = Info)
-	FString FactoryId;
-
 private:
+
 	struct FImpl;
 	TPimplPtr<FImpl> ImplPtr;
 
-	virtual TMap<FString, FPropertyValue> GetCreationParameters() PURE_VIRTUAL(ULiveLinkHubCaptureDevice::GetCreationParameters, return {}; );
+	friend ULiveLinkHubCaptureDeviceFactory;
 };
 
diff --git a/Engine/Plugins/VirtualProduction/CaptureManager/CaptureManagerApp/Source/CaptureManagerEditor/Private/Scripting/ULiveLinkHubNetworkCaptureDevice.cpp b/Engine/Plugins/VirtualProduction/CaptureManager/CaptureManagerApp/Source/CaptureManagerEditor/Private/Scripting/ULiveLinkHubNetworkCaptureDevice.cpp
deleted file mode 100644
diff --git a/Engine/Plugins/VirtualProduction/CaptureManager/CaptureManagerApp/Source/CaptureManagerEditor/Private/Scripting/ULiveLinkHubNetworkCaptureDevice.h b/Engine/Plugins/VirtualProduction/CaptureManager/CaptureManagerApp/Source/CaptureManagerEditor/Private/Scripting/ULiveLinkHubNetworkCaptureDevice.h
deleted file mode 100644
diff --git a/Engine/Plugins/VirtualProduction/CaptureManager/CaptureManagerApp/Source/CaptureManagerEditor/Private/UI/CaptureManagerPanelController.cpp b/Engine/Plugins/VirtualProduction/CaptureManager/CaptureManagerApp/Source/CaptureManagerEditor/Private/UI/CaptureManagerPanelController.cpp
--- a/Engine/Plugins/VirtualProduction/CaptureManager/CaptureManagerApp/Source/CaptureManagerEditor/Private/UI/CaptureManagerPanelController.cpp
+++ b/Engine/Plugins/VirtualProduction/CaptureManager/CaptureManagerApp/Source/CaptureManagerEditor/Private/UI/CaptureManagerPanelController.cpp
@@ -2,7 +2,6 @@
 
 #include "CaptureManagerPanelController.h"
 
-#include "CaptureSourceFrameworkModule.h"
 #include "TakesView.h"
 #include "SIngestJobProcessor.h"
 
@@ -10,7 +9,6 @@
 
 #include "IDetailsView.h"
 #include "PropertyEditorModule.h"
-#include "CaptureSourceObject.h"
 #include "Serialization/BufferArchive.h"
 
 #include "Misc/FileHelper.h"
diff --git a/Engine/Plugins/VirtualProduction/CaptureManager/CaptureManagerApp/Source/CaptureManagerEditor/Private/UI/CaptureSourceObject.cpp b/Engine/Plugins/VirtualProduction/CaptureManager/CaptureManagerApp/Source/CaptureManagerEditor/Private/UI/CaptureSourceObject.cpp
deleted file mode 100644
diff --git a/Engine/Plugins/VirtualProduction/CaptureManager/CaptureManagerApp/Source/CaptureManagerEditor/Private/UI/CaptureSourceObject.h b/Engine/Plugins/VirtualProduction/CaptureManager/CaptureManagerApp/Source/CaptureManagerEditor/Private/UI/CaptureSourceObject.h
deleted file mode 100644
diff --git a/Engine/Plugins/VirtualProduction/CaptureManager/CaptureManagerApp/Source/CaptureManagerEditor/Private/UI/CaptureSourceVirtualAsset.h b/Engine/Plugins/VirtualProduction/CaptureManager/CaptureManagerApp/Source/CaptureManagerEditor/Private/UI/CaptureSourceVirtualAsset.h
--- a/Engine/Plugins/VirtualProduction/CaptureManager/CaptureManagerApp/Source/CaptureManagerEditor/Private/UI/CaptureSourceVirtualAsset.h
+++ b/Engine/Plugins/VirtualProduction/CaptureManager/CaptureManagerApp/Source/CaptureManagerEditor/Private/UI/CaptureSourceVirtualAsset.h
@@ -1,7 +1,8 @@
 // Copyright Epic Games, Inc. All Rights Reserved.
+
 #pragma once
 
-#include "CaptureSourceManager.h"
+#include "Misc/Guid.h"
 
 #include "CaptureSourceVirtualAsset.generated.h"
 
diff --git a/Engine/Plugins/VirtualProduction/CaptureManager/CaptureManagerApp/Source/GenericLiveLinkDevice/Private/Archive/GenericLiveLinkArchiveDevice.cpp b/Engine/Plugins/VirtualProduction/CaptureManager/CaptureManagerApp/Source/GenericLiveLinkDevice/Private/Archive/GenericLiveLinkArchiveDevice.cpp
--- a/Engine/Plugins/VirtualProduction/CaptureManager/CaptureManagerApp/Source/GenericLiveLinkDevice/Private/Archive/GenericLiveLinkArchiveDevice.cpp
+++ b/Engine/Plugins/VirtualProduction/CaptureManager/CaptureManagerApp/Source/GenericLiveLinkDevice/Private/Archive/GenericLiveLinkArchiveDevice.cpp
@@ -46,18 +46,6 @@ void UGenericLiveLinkArchiveDevice::OnDeviceAdded(const FGuid InDeviceGuid, ULiv
 	DeviceGuid = InDeviceGuid;
 
 	CommunicationManager = UE::CaptureManager::Private::GetCommunicationManager();
-
-	UGenericLiveLinkArchiveDeviceSettings* ConfiguredSettings = CastChecked<UGenericLiveLinkArchiveDeviceSettings>(InSettings);
-	FString Path = ConfiguredSettings->StoragePath.Path;
-
-	if (!Path.IsEmpty() && FPaths::DirectoryExists(Path))
-	{
-		SetConnectionStatus(ELiveLinkDeviceConnectionStatus::Connected);
-	}
-	else
-	{
-		SetConnectionStatus(ELiveLinkDeviceConnectionStatus::Disconnected);
-	}
 }
 
 void UGenericLiveLinkArchiveDevice::OnSettingChanged(const FPropertyChangedEvent& InPropertyChangedEvent)
@@ -154,7 +142,7 @@ void UGenericLiveLinkArchiveDevice::RunIngestTakeAsync(int32 InTakeId, const UIn
 
 		FString TakeName = TakeMetadata.Slate + TEXT("_") + FString::FromInt(TakeMetadata.TakeNumber);
 
-		const FString TakeConversionDirectory = IngestOptions->Location;
+		const FString TakeConversionDirectory = IngestOptions->WorkingDirectory;
 		const FString StoragePath = GetSettings()->StoragePath.Path;
 
 		FGuid TakeUploadId = FGuid::NewGuid();
@@ -251,7 +239,7 @@ void UGenericLiveLinkArchiveDevice::CancelIngest(int32 InTakeId)
 
 			const FContext& Context = FoundContext->Value;
 
-			const FString TakeDirectory = FPaths::Combine(Context.IngestOptions->Location,
+			const FString TakeDirectory = FPaths::Combine(Context.IngestOptions->WorkingDirectory,
 														  Context.TakeName + TEXT("_") + FoundContext->Key.ToString());
 			IFileManager::Get().DeleteDirectory(*TakeDirectory, false, true);
 
@@ -290,11 +278,21 @@ bool UGenericLiveLinkArchiveDevice::SetHardwareId_Implementation(const FString&
 
 bool UGenericLiveLinkArchiveDevice::Connect_Implementation()
 {
-	return true;
+	const UGenericLiveLinkArchiveDeviceSettings* DeviceSettings = GetSettings();
+	FString Path = DeviceSettings->StoragePath.Path;
+
+	if (!Path.IsEmpty() && FPaths::DirectoryExists(Path))
+	{
+		SetConnectionStatus(ELiveLinkDeviceConnectionStatus::Connected);
+		return true;
+	}
+	
+	return false;
 }
 
 bool UGenericLiveLinkArchiveDevice::Disconnect_Implementation()
 {
+	SetConnectionStatus(ELiveLinkDeviceConnectionStatus::Disconnected);
 	return true;
 }
 
@@ -369,7 +367,7 @@ void UGenericLiveLinkArchiveDevice::UploadTakeFinished(const FGuid& InTakeUpload
 		const bool bShouldCleanConversionDirectory = GetDefault<UCaptureManagerSettings>()->bShouldCleanConversionDirectory;
 		if (bShouldCleanConversionDirectory)
 		{
-			const FString ConversionDirectory = FPaths::Combine(Context->IngestOptions->Location,
+			const FString ConversionDirectory = FPaths::Combine(Context->IngestOptions->WorkingDirectory,
 																Context->TakeName + TEXT("_") + InTakeUploadId.ToString());
 			IFileManager::Get().DeleteDirectory(*ConversionDirectory, false, true);
 		}
diff --git a/Engine/Plugins/VirtualProduction/CaptureManager/CaptureManagerApp/Source/GenericLiveLinkDevice/Private/Archive/GenericLiveLinkArchiveDevice.h b/Engine/Plugins/VirtualProduction/CaptureManager/CaptureManagerApp/Source/GenericLiveLinkDevice/Private/Archive/GenericLiveLinkArchiveDevice.h
--- a/Engine/Plugins/VirtualProduction/CaptureManager/CaptureManagerApp/Source/GenericLiveLinkDevice/Private/Archive/GenericLiveLinkArchiveDevice.h
+++ b/Engine/Plugins/VirtualProduction/CaptureManager/CaptureManagerApp/Source/GenericLiveLinkDevice/Private/Archive/GenericLiveLinkArchiveDevice.h
@@ -19,7 +19,7 @@
 #include "GenericLiveLinkArchiveDevice.generated.h"
 
 
-UCLASS()
+UCLASS(BlueprintType)
 class UGenericLiveLinkArchiveDeviceSettings : public ULiveLinkDeviceSettings
 {
 	GENERATED_BODY()
@@ -33,7 +33,7 @@ public:
 };
 
 
-UCLASS()
+UCLASS(BlueprintType)
 class UGenericLiveLinkArchiveDevice : public ULiveLinkDevice
 	, public ILiveLinkDeviceCapability_Ingest
 	, public ILiveLinkDeviceCapability_Connection
@@ -41,6 +41,7 @@ class UGenericLiveLinkArchiveDevice : public ULiveLinkDevice
 	GENERATED_BODY()
 
 public:
+	UFUNCTION(BlueprintCallable, Category = "GenericArchiveDeviceSettings")
 	const UGenericLiveLinkArchiveDeviceSettings* GetSettings() const;
 
 	//~ Begin ULiveLinkDevice interface
diff --git a/Engine/Plugins/VirtualProduction/CaptureManager/CaptureManagerApp/Source/GenericLiveLinkDevice/Private/CPS/GenericLiveLinkCPSDevice.cpp b/Engine/Plugins/VirtualProduction/CaptureManager/CaptureManagerApp/Source/GenericLiveLinkDevice/Private/CPS/GenericLiveLinkCPSDevice.cpp
--- a/Engine/Plugins/VirtualProduction/CaptureManager/CaptureManagerApp/Source/GenericLiveLinkDevice/Private/CPS/GenericLiveLinkCPSDevice.cpp
+++ b/Engine/Plugins/VirtualProduction/CaptureManager/CaptureManagerApp/Source/GenericLiveLinkDevice/Private/CPS/GenericLiveLinkCPSDevice.cpp
@@ -18,10 +18,6 @@
 
 #include "Utils/GenericLiveLinkDeviceUtility.h"
 
-UGenericLiveLinkCPSDeviceSettings* UGenericLiveLinkCPSDevice::GetSettings()
-{
-	return CastChecked<UGenericLiveLinkCPSDeviceSettings>(Settings);
-}
 const UGenericLiveLinkCPSDeviceSettings* UGenericLiveLinkCPSDevice::GetSettings() const
 {
 	return CastChecked<UGenericLiveLinkCPSDeviceSettings>(Settings);
@@ -174,7 +170,7 @@ void UGenericLiveLinkCPSDevice::CancelIngest(int32 InTakeId)
 
 			const FContext& Context = FoundContext->Value;
 
-			const FString TakeDirectory = FPaths::Combine(Context.IngestOptions->Location,
+			const FString TakeDirectory = FPaths::Combine(Context.IngestOptions->WorkingDirectory,
 														  Context.TakeName + TEXT("_") + FoundContext->Key.ToString());
 			IFileManager::Get().DeleteDirectory(*TakeDirectory, false, true);
 
@@ -219,7 +215,7 @@ bool UGenericLiveLinkCPSDevice::SetHardwareId_Implementation(const FString& Hard
 
 bool UGenericLiveLinkCPSDevice::Connect_Implementation()
 {
-	UGenericLiveLinkCPSDeviceSettings* DeviceSettings = GetSettings();
+	const UGenericLiveLinkCPSDeviceSettings* DeviceSettings = GetSettings();
 
 	using namespace UE::CaptureManager;
 
@@ -467,7 +463,7 @@ void UGenericLiveLinkCPSDevice::ConvertData(const FString& InTakeName,
 		return;
 	}
 
-	const FString TakeConversionDirectory = InIngestOptions->Location;
+	const FString TakeConversionDirectory = InIngestOptions->WorkingDirectory;
 	FGuid TakeUploadId = FGuid::NewGuid();
 
 	FCaptureDataConverterParams Params;
@@ -590,7 +586,7 @@ void UGenericLiveLinkCPSDevice::UploadTakeFinished(const FGuid& InTakeUploadId,
 		const bool bShouldCleanConversionDirectory = GetDefault<UCaptureManagerSettings>()->bShouldCleanConversionDirectory;
 		if (bShouldCleanConversionDirectory)
 		{
-			const FString ConversionDirectory = FPaths::Combine(Context->IngestOptions->Location,
+			const FString ConversionDirectory = FPaths::Combine(Context->IngestOptions->WorkingDirectory,
 																Context->TakeName + TEXT("_") + InTakeUploadId.ToString());
 			IFileManager::Get().DeleteDirectory(*ConversionDirectory, false, true);
 		}
diff --git a/Engine/Plugins/VirtualProduction/CaptureManager/CaptureManagerApp/Source/GenericLiveLinkDevice/Private/CPS/GenericLiveLinkCPSDevice.h b/Engine/Plugins/VirtualProduction/CaptureManager/CaptureManagerApp/Source/GenericLiveLinkDevice/Private/CPS/GenericLiveLinkCPSDevice.h
--- a/Engine/Plugins/VirtualProduction/CaptureManager/CaptureManagerApp/Source/GenericLiveLinkDevice/Private/CPS/GenericLiveLinkCPSDevice.h
+++ b/Engine/Plugins/VirtualProduction/CaptureManager/CaptureManagerApp/Source/GenericLiveLinkDevice/Private/CPS/GenericLiveLinkCPSDevice.h
@@ -23,7 +23,7 @@
 #include "GenericLiveLinkCPSDevice.generated.h"
 
 
-UCLASS()
+UCLASS(BlueprintType)
 class UGenericLiveLinkCPSDeviceSettings : public ULiveLinkDeviceSettings
 {
 	GENERATED_BODY()
@@ -43,7 +43,7 @@ public:
 };
 
 
-UCLASS()
+UCLASS(BlueprintType)
 class UGenericLiveLinkCPSDevice : public ULiveLinkDevice
 	, public ILiveLinkDeviceCapability_Ingest
 	, public ILiveLinkDeviceCapability_Connection
@@ -52,7 +52,7 @@ class UGenericLiveLinkCPSDevice : public ULiveLinkDevice
 	GENERATED_BODY()
 
 public:
-	UGenericLiveLinkCPSDeviceSettings* GetSettings();
+	UFUNCTION(BlueprintCallable, Category = "GenericCPSDeviceSettings")
 	const UGenericLiveLinkCPSDeviceSettings* GetSettings() const;
 
 	//~ Begin ULiveLinkDevice interface
diff --git a/Engine/Plugins/VirtualProduction/CaptureManager/CaptureManagerApp/Source/GenericLiveLinkDevice/Private/GenericLiveLinkDeviceHelper.cpp b/Engine/Plugins/VirtualProduction/CaptureManager/CaptureManagerApp/Source/GenericLiveLinkDevice/Private/GenericLiveLinkDeviceHelper.cpp
new file mode 100644
--- /dev/null
+++ b/Engine/Plugins/VirtualProduction/CaptureManager/CaptureManagerApp/Source/GenericLiveLinkDevice/Private/GenericLiveLinkDeviceHelper.cpp
@@ -0,0 +1,59 @@
+// Copyright Epic Games, Inc. All Rights Reserved.
+
+#include "GenericLiveLinkDeviceHelper.h"
+
+#include "Engine/Engine.h"
+
+#include "LiveLinkDeviceSubsystem.h"
+
+#include "Ingest/LiveLinkDeviceCapability_Ingest.h"
+
+#include "Archive/GenericLiveLinkArchiveDevice.h"
+#include "CPS/GenericLiveLinkCPSDevice.h"
+
+ULiveLinkDevice* UGenericLiveLinkDeviceHelper::CreateArchiveDevice(FString InName, FString InStoragePath) const
+{
+	TSubclassOf<ULiveLinkDevice> DeviceClass = UGenericLiveLinkArchiveDevice::StaticClass();
+
+	ULiveLinkDeviceSubsystem* DeviceSubsystem = GEngine->GetEngineSubsystem<ULiveLinkDeviceSubsystem>();
+	ULiveLinkDevice* NewDevice = NewObject<ULiveLinkDevice>(DeviceSubsystem, DeviceClass);
+
+	check(NewDevice->Implements<ULiveLinkDeviceCapability_Ingest>());
+
+	const FGuid DeviceGuid = DeviceSubsystem->AddDevice(NewDevice);
+
+	TSubclassOf<ULiveLinkDeviceSettings> SettingsClass = NewDevice->StaticGetSettingsClass();
+	ULiveLinkDeviceSettings* NewSettings = NewObject<ULiveLinkDeviceSettings>(NewDevice, SettingsClass);
+
+	UGenericLiveLinkArchiveDeviceSettings* DeviceSettings = CastChecked<UGenericLiveLinkArchiveDeviceSettings>(NewSettings);
+	DeviceSettings->StoragePath.Path = MoveTemp(InStoragePath);
+	DeviceSettings->DisplayName = MoveTemp(InName);
+
+	NewDevice->OnDeviceAdded(DeviceGuid, NewSettings);
+
+	return NewDevice;
+}
+
+ULiveLinkDevice* UGenericLiveLinkDeviceHelper::CreateNetworkDevice(FString InName, FString InIPAddress, int32 InPort) const
+{
+	TSubclassOf<ULiveLinkDevice> DeviceClass = UGenericLiveLinkCPSDevice::StaticClass();
+
+	ULiveLinkDeviceSubsystem* DeviceSubsystem = GEngine->GetEngineSubsystem<ULiveLinkDeviceSubsystem>();
+	ULiveLinkDevice* NewDevice = NewObject<ULiveLinkDevice>(DeviceSubsystem, DeviceClass);
+
+	check(NewDevice->Implements<ULiveLinkDeviceCapability_Ingest>());
+
+	const FGuid DeviceGuid = DeviceSubsystem->AddDevice(NewDevice);
+
+	TSubclassOf<ULiveLinkDeviceSettings> SettingsClass = NewDevice->StaticGetSettingsClass();
+	ULiveLinkDeviceSettings* NewSettings = NewObject<ULiveLinkDeviceSettings>(NewDevice, SettingsClass);
+
+	UGenericLiveLinkCPSDeviceSettings* DeviceSettings = CastChecked<UGenericLiveLinkCPSDeviceSettings>(NewSettings);
+	DeviceSettings->IpAddress.IpAddressString = MoveTemp(InIPAddress);
+	DeviceSettings->Port = static_cast<uint16>(InPort);
+	DeviceSettings->DisplayName = MoveTemp(InName);
+
+	NewDevice->OnDeviceAdded(DeviceGuid, NewSettings);
+
+	return NewDevice;
+}
\ No newline at end of file
diff --git a/Engine/Plugins/VirtualProduction/CaptureManager/CaptureManagerApp/Source/GenericLiveLinkDevice/Public/GenericLiveLinkDeviceHelper.h b/Engine/Plugins/VirtualProduction/CaptureManager/CaptureManagerApp/Source/GenericLiveLinkDevice/Public/GenericLiveLinkDeviceHelper.h
new file mode 100644
--- /dev/null
+++ b/Engine/Plugins/VirtualProduction/CaptureManager/CaptureManagerApp/Source/GenericLiveLinkDevice/Public/GenericLiveLinkDeviceHelper.h
@@ -0,0 +1,21 @@
+// Copyright Epic Games, Inc. All Rights Reserved.
+
+#pragma once
+
+#include "LiveLinkDevice.h"
+
+#include "GenericLiveLinkDeviceHelper.generated.h"
+
+UCLASS(BlueprintType)
+class GENERICLIVELINKDEVICE_API UGenericLiveLinkDeviceHelper : public UObject
+{
+	GENERATED_BODY()
+
+public:
+
+	UFUNCTION(BlueprintCallable, Category = Helper)
+	ULiveLinkDevice* CreateArchiveDevice(FString InName, FString InStoragePath) const;
+
+	UFUNCTION(BlueprintCallable, Category = Helper)
+	ULiveLinkDevice* CreateNetworkDevice(FString InName, FString InIPAddress, int32 InPort) const;
+};
\ No newline at end of file
diff --git a/Engine/Plugins/VirtualProduction/CaptureManager/CaptureManagerApp/Source/LiveLinkCapabilities/Private/Ingest/IngestCapability_ProcessHandle.cpp b/Engine/Plugins/VirtualProduction/CaptureManager/CaptureManagerApp/Source/LiveLinkCapabilities/Private/Ingest/IngestCapability_ProcessHandle.cpp
--- a/Engine/Plugins/VirtualProduction/CaptureManager/CaptureManagerApp/Source/LiveLinkCapabilities/Private/Ingest/IngestCapability_ProcessHandle.cpp
+++ b/Engine/Plugins/VirtualProduction/CaptureManager/CaptureManagerApp/Source/LiveLinkCapabilities/Private/Ingest/IngestCapability_ProcessHandle.cpp
@@ -55,13 +55,17 @@ bool UIngestCapability_ProcessResult::IsError() const
 	return !IsValid();
 }
 
-UIngestCapability_ProcessResult* UIngestCapability_ProcessResult::Create(int32 InCode, FText InMessage)
+UIngestCapability_ProcessResult* UIngestCapability_ProcessResult::Success()
 {
-	UIngestCapability_ProcessResult* Result =
-		NewObject<UIngestCapability_ProcessResult>();
+	return NewObject<UIngestCapability_ProcessResult>();
+}
 
-	Result->Code = InCode;
-	Result->Message = InMessage;
+UIngestCapability_ProcessResult* UIngestCapability_ProcessResult::Error(FText InMessage, int32 InCode)
+{
+	UIngestCapability_ProcessResult* Result = NewObject<UIngestCapability_ProcessResult>();
 
+	Result->Message = InMessage;
+	Result->Code = InCode;
+	
 	return Result;
-}
\ No newline at end of file
+}
diff --git a/Engine/Plugins/VirtualProduction/CaptureManager/CaptureManagerApp/Source/LiveLinkCapabilities/Private/Ingest/LiveLinkDeviceCapability_Ingest.cpp b/Engine/Plugins/VirtualProduction/CaptureManager/CaptureManagerApp/Source/LiveLinkCapabilities/Private/Ingest/LiveLinkDeviceCapability_Ingest.cpp
--- a/Engine/Plugins/VirtualProduction/CaptureManager/CaptureManagerApp/Source/LiveLinkCapabilities/Private/Ingest/LiveLinkDeviceCapability_Ingest.cpp
+++ b/Engine/Plugins/VirtualProduction/CaptureManager/CaptureManagerApp/Source/LiveLinkCapabilities/Private/Ingest/LiveLinkDeviceCapability_Ingest.cpp
@@ -68,7 +68,9 @@ UIngestCapability_TakeInformation* ILiveLinkDeviceCapability_Ingest::GetTakeInfo
 
 	FTakeMetadata TakeMetadata = TakeMetadataOpt.GetValue();
 
-	UIngestCapability_TakeInformation* TakeInformation = NewObject<UIngestCapability_TakeInformation>();
+	UIngestCapability_TakeInformation* TakeInformation = 
+		NewObject<UIngestCapability_TakeInformation>();
+
 	TakeInformation->SlateName = TakeMetadata.Slate;
 	TakeInformation->TakeNumber = TakeMetadata.TakeNumber;
 
@@ -114,7 +116,7 @@ void ILiveLinkDeviceCapability_Ingest::ExecuteProcessFinishedReporter(int32 InTa
 	Lock.Unlock();
 
 	UIngestCapability_ProcessResult* Result =
-		UIngestCapability_ProcessResult::Create();
+		UIngestCapability_ProcessResult::Success();
 
 	if (InMaybeError.HasError())
 	{
diff --git a/Engine/Plugins/VirtualProduction/CaptureManager/CaptureManagerApp/Source/LiveLinkCapabilities/Public/Ingest/IngestCapability_Options.h b/Engine/Plugins/VirtualProduction/CaptureManager/CaptureManagerApp/Source/LiveLinkCapabilities/Public/Ingest/IngestCapability_Options.h
--- a/Engine/Plugins/VirtualProduction/CaptureManager/CaptureManagerApp/Source/LiveLinkCapabilities/Public/Ingest/IngestCapability_Options.h
+++ b/Engine/Plugins/VirtualProduction/CaptureManager/CaptureManagerApp/Source/LiveLinkCapabilities/Public/Ingest/IngestCapability_Options.h
@@ -24,46 +24,46 @@ enum class EIngestCapability_ImagePixelFormat
 	Default = U8_RGB	UMETA(DisplayName = "Default (U8_RGB)")
 };
 
-USTRUCT()
+USTRUCT(BlueprintType)
 struct LIVELINKCAPABILITIES_API FIngestCapability_VideoOptions
 {
 	GENERATED_USTRUCT_BODY()
 
-	UPROPERTY(EditAnywhere, Category = "Live Link Device|Ingest")
+	UPROPERTY(EditAnywhere, BlueprintReadWrite, Category = "Live Link Device|Ingest")
 	FString FileNamePrefix;
 
-	UPROPERTY(EditAnywhere, Category = "Live Link Device|Ingest")
+	UPROPERTY(EditAnywhere, BlueprintReadWrite, Category = "Live Link Device|Ingest")
 	FString Format;
 
-	UPROPERTY(EditAnywhere, Category = "Live Link Device|Ingest")
+	UPROPERTY(EditAnywhere, BlueprintReadWrite, Category = "Live Link Device|Ingest")
 	EIngestCapability_ImagePixelFormat PixelFormat;
 };
 
-USTRUCT()
+USTRUCT(BlueprintType)
 struct LIVELINKCAPABILITIES_API FIngestCapability_AudioOptions
 {
 	GENERATED_USTRUCT_BODY()
 
-	UPROPERTY(EditAnywhere, Category = "Live Link Device|Ingest")
+	UPROPERTY(EditAnywhere, BlueprintReadWrite, Category = "Live Link Device|Ingest")
 	FString FileNamePrefix;
 
-	UPROPERTY(EditAnywhere, Category = "Live Link Device|Ingest")
+	UPROPERTY(EditAnywhere, BlueprintReadWrite, Category = "Live Link Device|Ingest")
 	FString Format;
 };
 
-UCLASS()
+UCLASS(BlueprintType)
 class LIVELINKCAPABILITIES_API UIngestCapability_Options : public UObject
 {
 	GENERATED_BODY()
 
 public:
 
-	UPROPERTY(EditAnywhere, Category = "Live Link Device|Ingest")
-	FString Location;
+	UPROPERTY(EditAnywhere, BlueprintReadWrite, Category = "Live Link Device|Ingest")
+	FString WorkingDirectory;
 
-	UPROPERTY(EditAnywhere, Category = "Live Link Device|Ingest")
+	UPROPERTY(EditAnywhere, BlueprintReadWrite, Category = "Live Link Device|Ingest")
 	FIngestCapability_VideoOptions Video;
 
-	UPROPERTY(EditAnywhere, Category = "Live Link Device|Ingest")
+	UPROPERTY(EditAnywhere, BlueprintReadWrite, Category = "Live Link Device|Ingest")
 	FIngestCapability_AudioOptions Audio;
 };
\ No newline at end of file
diff --git a/Engine/Plugins/VirtualProduction/CaptureManager/CaptureManagerApp/Source/LiveLinkCapabilities/Public/Ingest/IngestCapability_ProcessHandle.h b/Engine/Plugins/VirtualProduction/CaptureManager/CaptureManagerApp/Source/LiveLinkCapabilities/Public/Ingest/IngestCapability_ProcessHandle.h
--- a/Engine/Plugins/VirtualProduction/CaptureManager/CaptureManagerApp/Source/LiveLinkCapabilities/Public/Ingest/IngestCapability_ProcessHandle.h
+++ b/Engine/Plugins/VirtualProduction/CaptureManager/CaptureManagerApp/Source/LiveLinkCapabilities/Public/Ingest/IngestCapability_ProcessHandle.h
@@ -46,26 +46,30 @@ DECLARE_DYNAMIC_DELEGATE_TwoParams(FProcessProgressReporter, int32, TakeIdentifi
 using FIngestProcessProgressReporter = UE::CaptureManager::TManagedDelegate<int32, double>;
 using FIngestProcessFinishReporter = UE::CaptureManager::TManagedDelegate<int32, TValueOrError<void, FIngestCapability_Error>>;
 
-UCLASS()
+UCLASS(BlueprintType)
 class LIVELINKCAPABILITIES_API UIngestCapability_ProcessResult : public UObject
 {
 	GENERATED_BODY()
 
 public:
 
-	UPROPERTY()
+	UPROPERTY(VisibleAnywhere, BlueprintReadOnly, Category = "Live Link Device|Ingest")
 	FText Message;
 
-	UPROPERTY()
+	UPROPERTY(VisibleAnywhere, BlueprintReadOnly, Category = "Live Link Device|Ingest")
 	int32 Code = 0;
 
+	UFUNCTION(BlueprintCallable, Category = "Live Link Device|Ingest")
 	bool IsValid() const;
+
+	UFUNCTION(BlueprintCallable, Category = "Live Link Device|Ingest")
 	bool IsError() const;
 
-	static UIngestCapability_ProcessResult* Create(int32 InCode = 0, FText InMessage = FText());
+	static UIngestCapability_ProcessResult* Success();
+	static UIngestCapability_ProcessResult* Error(FText InMessage, int32 InCode = -1);
 };
 
-UCLASS()
+UCLASS(BlueprintType)
 class LIVELINKCAPABILITIES_API UIngestCapability_ProcessHandle : public UObject
 {
 	GENERATED_BODY()
diff --git a/Engine/Plugins/VirtualProduction/CaptureManager/CaptureManagerApp/Source/LiveLinkCapabilities/Public/Ingest/IngestCapability_TakeInformation.h b/Engine/Plugins/VirtualProduction/CaptureManager/CaptureManagerApp/Source/LiveLinkCapabilities/Public/Ingest/IngestCapability_TakeInformation.h
--- a/Engine/Plugins/VirtualProduction/CaptureManager/CaptureManagerApp/Source/LiveLinkCapabilities/Public/Ingest/IngestCapability_TakeInformation.h
+++ b/Engine/Plugins/VirtualProduction/CaptureManager/CaptureManagerApp/Source/LiveLinkCapabilities/Public/Ingest/IngestCapability_TakeInformation.h
@@ -11,23 +11,30 @@ namespace UE::CaptureManager
 using FTakeId = int32;
 }
 
-UCLASS()
-class LIVELINKCAPABILITIES_API UIngestCapability_TakeInformation : public UObject
+UCLASS(BlueprintType)
+class LIVELINKCAPABILITIES_API UIngestCapability_TakeInformation 
+	: public UObject
 {
 	GENERATED_BODY()
 
 public:
 
-	UPROPERTY()
+	UPROPERTY(VisibleAnywhere, BlueprintReadOnly, Category = TakeInfo)
 	FString DeviceName;
 
-	UPROPERTY()
+	UPROPERTY(VisibleAnywhere, BlueprintReadOnly, Category = TakeInfo)
 	FString SlateName;
 
-	UPROPERTY()
+	UPROPERTY(VisibleAnywhere, BlueprintReadOnly, Category = TakeInfo)
 	int32 TakeNumber = -1;
 
-	UPROPERTY()
+	UPROPERTY(VisibleAnywhere, BlueprintReadOnly, Category = TakeInfo)
 	FDateTime DateTime;
+
+	UFUNCTION(BlueprintCallable, Category = TakeInfo)
+	FString GetDateTimeString() const
+	{
+		return DateTime.ToIso8601();
+	}
 };
 
