Fix localization issues in Capture Manager

diff --git a/Engine/Plugins/VirtualProduction/CaptureManager/CaptureManagerApp/Source/CaptureManagerEditor/Private/UI/LiveLinkHubCaptureManagerMode.cpp b/Engine/Plugins/VirtualProduction/CaptureManager/CaptureManagerApp/Source/CaptureManagerEditor/Private/UI/LiveLinkHubCaptureManagerMode.cpp
--- a/Engine/Plugins/VirtualProduction/CaptureManager/CaptureManagerApp/Source/CaptureManagerEditor/Private/UI/LiveLinkHubCaptureManagerMode.cpp
+++ b/Engine/Plugins/VirtualProduction/CaptureManager/CaptureManagerApp/Source/CaptureManagerEditor/Private/UI/LiveLinkHubCaptureManagerMode.cpp
@@ -50,10 +50,6 @@ namespace UE::CaptureManager
 	}
 }
 
-const FText FLiveLinkHubCaptureManagerMode::TakeBrowserTabName = LOCTEXT("TakeBrowserTabLabel", "Take Browser");
-const FText FLiveLinkHubCaptureManagerMode::JobsListTabName = LOCTEXT("JobsListTabLabel", "Jobs");
-const FText FLiveLinkHubCaptureManagerMode::JobDetailsTabName = LOCTEXT("JobDetailsTabLabel", "Job Details");
-
 struct FTakeBrowserTabSummoner : public FWorkflowTabFactory
 {
 public:
diff --git a/Engine/Plugins/VirtualProduction/CaptureManager/CaptureManagerApp/Source/CaptureManagerEditor/Private/UI/LiveLinkHubCaptureManagerMode.h b/Engine/Plugins/VirtualProduction/CaptureManager/CaptureManagerApp/Source/CaptureManagerEditor/Private/UI/LiveLinkHubCaptureManagerMode.h
--- a/Engine/Plugins/VirtualProduction/CaptureManager/CaptureManagerApp/Source/CaptureManagerEditor/Private/UI/LiveLinkHubCaptureManagerMode.h
+++ b/Engine/Plugins/VirtualProduction/CaptureManager/CaptureManagerApp/Source/CaptureManagerEditor/Private/UI/LiveLinkHubCaptureManagerMode.h
@@ -66,11 +66,6 @@ private:
 
 	TSharedPtr<class FCaptureManagerPanelController> PanelController;
 
-	static const FText CaptureSourcesTabName;
-	static const FText TakeBrowserTabName;
-	static const FText JobsListTabName;
-	static const FText JobDetailsTabName;
-
 	/** Set of spawnable tabs in the mode */
 	FWorkflowAllowedTabSet TabFactories;
 	
diff --git a/Engine/Plugins/VirtualProduction/CaptureManager/CaptureManagerApp/Source/CaptureManagerMediaRW/Private/Utils/WindowsJpegDecoder.cpp b/Engine/Plugins/VirtualProduction/CaptureManager/CaptureManagerApp/Source/CaptureManagerMediaRW/Private/Utils/WindowsJpegDecoder.cpp
--- a/Engine/Plugins/VirtualProduction/CaptureManager/CaptureManagerApp/Source/CaptureManagerMediaRW/Private/Utils/WindowsJpegDecoder.cpp
+++ b/Engine/Plugins/VirtualProduction/CaptureManager/CaptureManagerApp/Source/CaptureManagerMediaRW/Private/Utils/WindowsJpegDecoder.cpp
@@ -86,7 +86,7 @@ TOptional<FText> FWindowsJpegDecoder::Decode(uint8* InData, uint32 InSize, TArra
 
 	WICPixelFormatGUID PixelFormat;
 	Result = Bitmap->GetPixelFormat(&PixelFormat);
-	WINJD_CHECK_AND_RETURN_ERROR_MESSAGE(Result, LOCTEXT("Decode_FailedObtainData", "Failed to obtain the pixel format"));
+	WINJD_CHECK_AND_RETURN_ERROR_MESSAGE(Result, LOCTEXT("Decode_FailedObtainPixelFormat", "Failed to obtain the pixel format"));
 
 	OutPixelFormat = ConvertPixelFormat(PixelFormat);
 
