Delete the Capture Source framework as it will not be used anymore

diff --git a/Engine/Plugins/VirtualProduction/CaptureManager/CaptureManagerApp/CaptureManagerApp.uplugin b/Engine/Plugins/VirtualProduction/CaptureManager/CaptureManagerApp/CaptureManagerApp.uplugin
--- a/Engine/Plugins/VirtualProduction/CaptureManager/CaptureManagerApp/CaptureManagerApp.uplugin
+++ b/Engine/Plugins/VirtualProduction/CaptureManager/CaptureManagerApp/CaptureManagerApp.uplugin
@@ -16,90 +16,65 @@
 	"IsExperimentalVersion": true,
 	"Installed": false,
 	"Hidden" :  true,
-	"SupportedPrograms": [ "LiveLinkHub" ],
 	"Modules": [
 		{
 			"Name": "CaptureManagerEditor",
-			"Type": "EditorAndProgram",
-			"ProgramAllowList": [ "LiveLinkHub" ],
+			"Type": "Editor",
 			"LoadingPhase": "Default"
 		},
 		{
 			"Name": "LiveLinkHubExportClient",
-			"Type": "EditorAndProgram",
-			"ProgramAllowList": [ "LiveLinkHub" ],
+			"Type": "Editor",
 			"LoadingPhase": "Default"
 		},
 		{
 			"Name": "LiveLinkFaceMetadata",
-			"Type": "EditorAndProgram",
-			"ProgramAllowList": [ "LiveLinkHub" ],
+			"Type": "Editor",
 			"LoadingPhase": "Default"
 		},
 		{
 			"Name": "StereoCameraMetadata",
-			"Type": "EditorAndProgram",
-			"ProgramAllowList": [ "LiveLinkHub" ],
-			"LoadingPhase": "Default"
-		},
-		{
-			"Name": "LiveLinkHubCaptureSourceWidgets",
-			"Type": "EditorAndProgram",
-			"ProgramAllowList": [ "LiveLinkHub" ],
+			"Type": "Editor",
 			"LoadingPhase": "Default"
 		},
 		{
 			"Name": "LiveLinkHubDiscovery",
-			"Type": "EditorAndProgram",
-			"ProgramAllowList": [ "LiveLinkHub" ],
+			"Type": "Editor",
 			"LoadingPhase": "Default"
 		},
 		{
 			"Name": "CaptureManagerPipeline",
-			"Type": "EditorAndProgram",
-			"ProgramAllowList": [ "LiveLinkHub" ],
+			"Type": "Editor",
 			"LoadingPhase": "Default"
 		},
 		{
 			"Name": "CaptureManagerMediaRW",
-			"Type": "EditorAndProgram",
-			"ProgramAllowList": [ "LiveLinkHub" ],
+			"Type": "Editor",
 			"LoadingPhase": "Default"
 		},
 		{
 			"Name": "CaptureDataConverter",
-			"Type": "EditorAndProgram",
-			"ProgramAllowList": [ "LiveLinkHub" ],
+			"Type": "Editor",
 			"LoadingPhase": "Default"
 		},
 		{
 			"Name": "CaptureCommunicationManager",
-			"Type": "EditorAndProgram",
-			"ProgramAllowList": [ "LiveLinkHub" ],
-			"LoadingPhase": "Default"
-		},
-		{
-			"Name": "GenericCaptureSource",
-			"Type": "EditorAndProgram",
-			"ProgramAllowList": [ "LiveLinkHub" ],
+			"Type": "Editor",
 			"LoadingPhase": "Default"
 		},
 		{
 			"Name": "CaptureManagerSettings",
-			"Type": "EditorAndProgram",
-			"ProgramAllowList": [ "LiveLinkHub" ],
+			"Type": "Editor",
 			"LoadingPhase": "Default"
 		},
 		{
 			"Name": "GenericLiveLinkDevice",
-			"Type": "EditorAndProgram",
-			"ProgramAllowList": [ "LiveLinkHub" ],
+			"Type": "Editor",
 			"LoadingPhase": "Default"
 		},
 		{
 			"Name": "LiveLinkCapabilities",
-			"Type": "EditorAndProgram",
-			"ProgramAllowList": [ "LiveLinkHub" ],
+			"Type": "Editor",
 			"LoadingPhase": "Default"
 		}
 	],
diff --git a/Engine/Plugins/VirtualProduction/CaptureManager/CaptureManagerApp/Source/GenericCaptureSource/GenericCaptureSource.Build.cs b/Engine/Plugins/VirtualProduction/CaptureManager/CaptureManagerApp/Source/GenericCaptureSource/GenericCaptureSource.Build.cs
deleted file mode 100644
diff --git a/Engine/Plugins/VirtualProduction/CaptureManager/CaptureManagerApp/Source/GenericCaptureSource/Private/Archive/GenericArchiveCaptureSource.cpp b/Engine/Plugins/VirtualProduction/CaptureManager/CaptureManagerApp/Source/GenericCaptureSource/Private/Archive/GenericArchiveCaptureSource.cpp
deleted file mode 100644
diff --git a/Engine/Plugins/VirtualProduction/CaptureManager/CaptureManagerApp/Source/GenericCaptureSource/Private/Archive/GenericArchiveCaptureSource.h b/Engine/Plugins/VirtualProduction/CaptureManager/CaptureManagerApp/Source/GenericCaptureSource/Private/Archive/GenericArchiveCaptureSource.h
deleted file mode 100644
diff --git a/Engine/Plugins/VirtualProduction/CaptureManager/CaptureManagerApp/Source/GenericCaptureSource/Private/Archive/GenericArchiveCaptureSourceFactory.cpp b/Engine/Plugins/VirtualProduction/CaptureManager/CaptureManagerApp/Source/GenericCaptureSource/Private/Archive/GenericArchiveCaptureSourceFactory.cpp
deleted file mode 100644
diff --git a/Engine/Plugins/VirtualProduction/CaptureManager/CaptureManagerApp/Source/GenericCaptureSource/Private/Archive/GenericArchiveCaptureSourceFactory.h b/Engine/Plugins/VirtualProduction/CaptureManager/CaptureManagerApp/Source/GenericCaptureSource/Private/Archive/GenericArchiveCaptureSourceFactory.h
deleted file mode 100644
diff --git a/Engine/Plugins/VirtualProduction/CaptureManager/CaptureManagerApp/Source/GenericCaptureSource/Private/GenericCaptureSourceLog.h b/Engine/Plugins/VirtualProduction/CaptureManager/CaptureManagerApp/Source/GenericCaptureSource/Private/GenericCaptureSourceLog.h
deleted file mode 100644
diff --git a/Engine/Plugins/VirtualProduction/CaptureManager/CaptureManagerApp/Source/GenericCaptureSource/Private/GenericCaptureSourceModule.cpp b/Engine/Plugins/VirtualProduction/CaptureManager/CaptureManagerApp/Source/GenericCaptureSource/Private/GenericCaptureSourceModule.cpp
deleted file mode 100644
diff --git a/Engine/Plugins/VirtualProduction/CaptureManager/CaptureManagerApp/Source/GenericCaptureSource/Private/GenericCaptureSourceModule.h b/Engine/Plugins/VirtualProduction/CaptureManager/CaptureManagerApp/Source/GenericCaptureSource/Private/GenericCaptureSourceModule.h
deleted file mode 100644
diff --git a/Engine/Plugins/VirtualProduction/CaptureManager/CaptureManagerApp/Source/GenericCaptureSource/Private/Live/CPS/CPSDataStream.cpp b/Engine/Plugins/VirtualProduction/CaptureManager/CaptureManagerApp/Source/GenericCaptureSource/Private/Live/CPS/CPSDataStream.cpp
deleted file mode 100644
diff --git a/Engine/Plugins/VirtualProduction/CaptureManager/CaptureManagerApp/Source/GenericCaptureSource/Private/Live/CPS/CPSDataStream.h b/Engine/Plugins/VirtualProduction/CaptureManager/CaptureManagerApp/Source/GenericCaptureSource/Private/Live/CPS/CPSDataStream.h
deleted file mode 100644
diff --git a/Engine/Plugins/VirtualProduction/CaptureManager/CaptureManagerApp/Source/GenericCaptureSource/Private/Live/CPS/CPSDevice.cpp b/Engine/Plugins/VirtualProduction/CaptureManager/CaptureManagerApp/Source/GenericCaptureSource/Private/Live/CPS/CPSDevice.cpp
deleted file mode 100644
diff --git a/Engine/Plugins/VirtualProduction/CaptureManager/CaptureManagerApp/Source/GenericCaptureSource/Private/Live/CPS/CPSDevice.h b/Engine/Plugins/VirtualProduction/CaptureManager/CaptureManagerApp/Source/GenericCaptureSource/Private/Live/CPS/CPSDevice.h
deleted file mode 100644
diff --git a/Engine/Plugins/VirtualProduction/CaptureManager/CaptureManagerApp/Source/GenericCaptureSource/Private/Live/CPS/CPSFileStream.cpp b/Engine/Plugins/VirtualProduction/CaptureManager/CaptureManagerApp/Source/GenericCaptureSource/Private/Live/CPS/CPSFileStream.cpp
deleted file mode 100644
diff --git a/Engine/Plugins/VirtualProduction/CaptureManager/CaptureManagerApp/Source/GenericCaptureSource/Private/Live/CPS/CPSFileStream.h b/Engine/Plugins/VirtualProduction/CaptureManager/CaptureManagerApp/Source/GenericCaptureSource/Private/Live/CPS/CPSFileStream.h
deleted file mode 100644
diff --git a/Engine/Plugins/VirtualProduction/CaptureManager/CaptureManagerApp/Source/GenericCaptureSource/Private/Live/GenericCPSCaptureSource.cpp b/Engine/Plugins/VirtualProduction/CaptureManager/CaptureManagerApp/Source/GenericCaptureSource/Private/Live/GenericCPSCaptureSource.cpp
deleted file mode 100644
diff --git a/Engine/Plugins/VirtualProduction/CaptureManager/CaptureManagerApp/Source/GenericCaptureSource/Private/Live/GenericCPSCaptureSource.h b/Engine/Plugins/VirtualProduction/CaptureManager/CaptureManagerApp/Source/GenericCaptureSource/Private/Live/GenericCPSCaptureSource.h
deleted file mode 100644
diff --git a/Engine/Plugins/VirtualProduction/CaptureManager/CaptureManagerApp/Source/GenericCaptureSource/Private/Live/GenericCPSCaptureSourceFactory.cpp b/Engine/Plugins/VirtualProduction/CaptureManager/CaptureManagerApp/Source/GenericCaptureSource/Private/Live/GenericCPSCaptureSourceFactory.cpp
deleted file mode 100644
diff --git a/Engine/Plugins/VirtualProduction/CaptureManager/CaptureManagerApp/Source/GenericCaptureSource/Private/Live/GenericCPSCaptureSourceFactory.h b/Engine/Plugins/VirtualProduction/CaptureManager/CaptureManagerApp/Source/GenericCaptureSource/Private/Live/GenericCPSCaptureSourceFactory.h
deleted file mode 100644
diff --git a/Engine/Plugins/VirtualProduction/CaptureManager/CaptureManagerApp/Source/GenericLiveLinkDevice/Private/CPS/Protocol/CPSDevice.h b/Engine/Plugins/VirtualProduction/CaptureManager/CaptureManagerApp/Source/GenericLiveLinkDevice/Private/CPS/Protocol/CPSDevice.h
--- a/Engine/Plugins/VirtualProduction/CaptureManager/CaptureManagerApp/Source/GenericLiveLinkDevice/Private/CPS/Protocol/CPSDevice.h
+++ b/Engine/Plugins/VirtualProduction/CaptureManager/CaptureManagerApp/Source/GenericLiveLinkDevice/Private/CPS/Protocol/CPSDevice.h
@@ -9,10 +9,10 @@
 #include "Async/EventSourceUtils.h"
 #include "Misc/Optional.h"
 
-#include "Capabilities/IngestCapability.h"
-
 #include "Async/CaptureTimerManager.h"
 
+#include "Ingest/IngestCapability_TakeInformation.h"
+
 namespace UE::CaptureManager
 {
 
diff --git a/Engine/Plugins/VirtualProduction/CaptureManager/CaptureManagerApp/Source/LiveLinkHubCaptureSourceWidgets/LiveLinkHubCaptureSourceWidgets.Build.cs b/Engine/Plugins/VirtualProduction/CaptureManager/CaptureManagerApp/Source/LiveLinkHubCaptureSourceWidgets/LiveLinkHubCaptureSourceWidgets.Build.cs
deleted file mode 100644
diff --git a/Engine/Plugins/VirtualProduction/CaptureManager/CaptureManagerApp/Source/LiveLinkHubCaptureSourceWidgets/Private/LiveLinkHubCaptureSourceWidgetsModule.cpp b/Engine/Plugins/VirtualProduction/CaptureManager/CaptureManagerApp/Source/LiveLinkHubCaptureSourceWidgets/Private/LiveLinkHubCaptureSourceWidgetsModule.cpp
deleted file mode 100644
diff --git a/Engine/Plugins/VirtualProduction/CaptureManager/CaptureManagerApp/Source/LiveLinkHubCaptureSourceWidgets/Private/LiveLinkHubCaptureSourceWidgetsModule.h b/Engine/Plugins/VirtualProduction/CaptureManager/CaptureManagerApp/Source/LiveLinkHubCaptureSourceWidgets/Private/LiveLinkHubCaptureSourceWidgetsModule.h
deleted file mode 100644
diff --git a/Engine/Plugins/VirtualProduction/CaptureManager/CaptureManagerApp/Source/LiveLinkHubCaptureSourceWidgets/Private/SArchiveCreationPanel.cpp b/Engine/Plugins/VirtualProduction/CaptureManager/CaptureManagerApp/Source/LiveLinkHubCaptureSourceWidgets/Private/SArchiveCreationPanel.cpp
deleted file mode 100644
diff --git a/Engine/Plugins/VirtualProduction/CaptureManager/CaptureManagerApp/Source/LiveLinkHubCaptureSourceWidgets/Private/SConnectionCreationPanel.cpp b/Engine/Plugins/VirtualProduction/CaptureManager/CaptureManagerApp/Source/LiveLinkHubCaptureSourceWidgets/Private/SConnectionCreationPanel.cpp
deleted file mode 100644
diff --git a/Engine/Plugins/VirtualProduction/CaptureManager/CaptureManagerApp/Source/LiveLinkHubCaptureSourceWidgets/Private/Utils.cpp b/Engine/Plugins/VirtualProduction/CaptureManager/CaptureManagerApp/Source/LiveLinkHubCaptureSourceWidgets/Private/Utils.cpp
deleted file mode 100644
diff --git a/Engine/Plugins/VirtualProduction/CaptureManager/CaptureManagerApp/Source/LiveLinkHubCaptureSourceWidgets/Private/Utils.h b/Engine/Plugins/VirtualProduction/CaptureManager/CaptureManagerApp/Source/LiveLinkHubCaptureSourceWidgets/Private/Utils.h
deleted file mode 100644
diff --git a/Engine/Plugins/VirtualProduction/CaptureManager/CaptureManagerApp/Source/LiveLinkHubCaptureSourceWidgets/Public/Widgets/SArchiveCreationPanel.h b/Engine/Plugins/VirtualProduction/CaptureManager/CaptureManagerApp/Source/LiveLinkHubCaptureSourceWidgets/Public/Widgets/SArchiveCreationPanel.h
deleted file mode 100644
diff --git a/Engine/Plugins/VirtualProduction/CaptureManager/CaptureManagerApp/Source/LiveLinkHubCaptureSourceWidgets/Public/Widgets/SConnectionCreationPanel.h b/Engine/Plugins/VirtualProduction/CaptureManager/CaptureManagerApp/Source/LiveLinkHubCaptureSourceWidgets/Public/Widgets/SConnectionCreationPanel.h
deleted file mode 100644
diff --git a/Engine/Plugins/VirtualProduction/CaptureManager/CaptureManagerApp/Source/LiveLinkHubDiscovery/LiveLinkHubDiscovery.Build.cs b/Engine/Plugins/VirtualProduction/CaptureManager/CaptureManagerApp/Source/LiveLinkHubDiscovery/LiveLinkHubDiscovery.Build.cs
--- a/Engine/Plugins/VirtualProduction/CaptureManager/CaptureManagerApp/Source/LiveLinkHubDiscovery/LiveLinkHubDiscovery.Build.cs
+++ b/Engine/Plugins/VirtualProduction/CaptureManager/CaptureManagerApp/Source/LiveLinkHubDiscovery/LiveLinkHubDiscovery.Build.cs
@@ -8,7 +8,6 @@ public class LiveLinkHubDiscovery : ModuleRules
 	{
 		PrivateDependencyModuleNames.AddRange(new string[] 
 		{
-			"CaptureSourceFramework",
 			"Core",
 			"CoreUObject",
 			"CaptureUtils",
diff --git a/Engine/Plugins/VirtualProduction/CaptureManager/CaptureManagerApp/Source/StereoCameraMetadata/Private/StereoCameraTakeMetadata.h b/Engine/Plugins/VirtualProduction/CaptureManager/CaptureManagerApp/Source/StereoCameraMetadata/Private/StereoCameraTakeMetadata.h
--- a/Engine/Plugins/VirtualProduction/CaptureManager/CaptureManagerApp/Source/StereoCameraMetadata/Private/StereoCameraTakeMetadata.h
+++ b/Engine/Plugins/VirtualProduction/CaptureManager/CaptureManagerApp/Source/StereoCameraMetadata/Private/StereoCameraTakeMetadata.h
@@ -7,7 +7,6 @@
 
 #include "Async/StopToken.h"
 
-#include "Capabilities/IngestCapability.h"
 #include "CaptureManagerTakeMetadata.h"
 
 class FStereoCameraInfo
diff --git a/Engine/Plugins/VirtualProduction/CaptureManager/CaptureManagerApp/Source/StereoCameraMetadata/StereoCameraMetadata.Build.cs b/Engine/Plugins/VirtualProduction/CaptureManager/CaptureManagerApp/Source/StereoCameraMetadata/StereoCameraMetadata.Build.cs
--- a/Engine/Plugins/VirtualProduction/CaptureManager/CaptureManagerApp/Source/StereoCameraMetadata/StereoCameraMetadata.Build.cs
+++ b/Engine/Plugins/VirtualProduction/CaptureManager/CaptureManagerApp/Source/StereoCameraMetadata/StereoCameraMetadata.Build.cs
@@ -16,7 +16,6 @@ public class StereoCameraMetadata : ModuleRules
 			"Engine",
 			"CaptureUtils",
 			"DataIngestCore",
-			"CaptureSourceFramework",
 			"CaptureManagerTakeMetadata",
 			"CaptureDataCore"
 		});
diff --git a/Engine/Plugins/VirtualProduction/CaptureManager/CaptureManagerCore/CaptureManagerCore.uplugin b/Engine/Plugins/VirtualProduction/CaptureManager/CaptureManagerCore/CaptureManagerCore.uplugin
--- a/Engine/Plugins/VirtualProduction/CaptureManager/CaptureManagerCore/CaptureManagerCore.uplugin
+++ b/Engine/Plugins/VirtualProduction/CaptureManager/CaptureManagerCore/CaptureManagerCore.uplugin
@@ -16,42 +16,30 @@
 	"IsExperimentalVersion": true,
 	"Installed": false,
 	"Hidden": true,
-	"SupportedPrograms": [ "LiveLinkHub" ],
 	"Modules": [
 		{
 			"Name": "CaptureUtils",
-			"Type": "EditorAndProgram",
-			"ProgramAllowList": [ "LiveLinkHub" ],
-			"LoadingPhase": "PreDefault"
-		},
-		{
-			"Name": "CaptureSourceFramework",
-			"Type": "EditorAndProgram",
-			"ProgramAllowList": [ "LiveLinkHub" ],
+			"Type": "Editor",
 			"LoadingPhase": "PreDefault"
 		},
 		{
 			"Name": "CaptureProtocolStack",
-			"Type": "EditorAndProgram",
-			"ProgramAllowList": [ "LiveLinkHub" ],
+			"Type": "Editor",
 			"LoadingPhase": "Default"
 		},
 		{
 			"Name": "DataIngestCore",
-			"Type": "EditorAndProgram",
-			"ProgramAllowList": [ "LiveLinkHub" ],
+			"Type": "Editor",
 			"LoadingPhase": "Default"
 		},
 		{
 			"Name": "LiveLinkHubCaptureMessaging",
-			"Type": "EditorAndProgram",
-			"ProgramAllowList": [ "LiveLinkHub" ],
+			"Type": "Editor",
 			"LoadingPhase": "Default"
 		},
 		{
 			"Name": "CaptureManagerTakeMetadata",
-			"Type": "EditorAndProgram",
-			"ProgramAllowList": [ "LiveLinkHub" ],
+			"Type": "Editor",
 			"LoadingPhase": "Default"
 		}		
 	],
diff --git a/Engine/Plugins/VirtualProduction/CaptureManager/CaptureManagerCore/Source/CaptureSourceFramework/CaptureSourceFramework.Build.cs b/Engine/Plugins/VirtualProduction/CaptureManager/CaptureManagerCore/Source/CaptureSourceFramework/CaptureSourceFramework.Build.cs
deleted file mode 100644
diff --git a/Engine/Plugins/VirtualProduction/CaptureManager/CaptureManagerCore/Source/CaptureSourceFramework/Private/Capabilities/DiskCapacityCapability.cpp b/Engine/Plugins/VirtualProduction/CaptureManager/CaptureManagerCore/Source/CaptureSourceFramework/Private/Capabilities/DiskCapacityCapability.cpp
deleted file mode 100644
diff --git a/Engine/Plugins/VirtualProduction/CaptureManager/CaptureManagerCore/Source/CaptureSourceFramework/Private/Capabilities/IngestCapability.cpp b/Engine/Plugins/VirtualProduction/CaptureManager/CaptureManagerCore/Source/CaptureSourceFramework/Private/Capabilities/IngestCapability.cpp
deleted file mode 100644
diff --git a/Engine/Plugins/VirtualProduction/CaptureManager/CaptureManagerCore/Source/CaptureSourceFramework/Private/Capabilities/RecordingCapability.cpp b/Engine/Plugins/VirtualProduction/CaptureManager/CaptureManagerCore/Source/CaptureSourceFramework/Private/Capabilities/RecordingCapability.cpp
deleted file mode 100644
diff --git a/Engine/Plugins/VirtualProduction/CaptureManager/CaptureManagerCore/Source/CaptureSourceFramework/Private/CaptureSource.cpp b/Engine/Plugins/VirtualProduction/CaptureManager/CaptureManagerCore/Source/CaptureSourceFramework/Private/CaptureSource.cpp
deleted file mode 100644
diff --git a/Engine/Plugins/VirtualProduction/CaptureManager/CaptureManagerCore/Source/CaptureSourceFramework/Private/CaptureSourceCapability.cpp b/Engine/Plugins/VirtualProduction/CaptureManager/CaptureManagerCore/Source/CaptureSourceFramework/Private/CaptureSourceCapability.cpp
deleted file mode 100644
diff --git a/Engine/Plugins/VirtualProduction/CaptureManager/CaptureManagerCore/Source/CaptureSourceFramework/Private/CaptureSourceCapabilityCommand.cpp b/Engine/Plugins/VirtualProduction/CaptureManager/CaptureManagerCore/Source/CaptureSourceFramework/Private/CaptureSourceCapabilityCommand.cpp
deleted file mode 100644
diff --git a/Engine/Plugins/VirtualProduction/CaptureManager/CaptureManagerCore/Source/CaptureSourceFramework/Private/CaptureSourceCapabilityEvent.cpp b/Engine/Plugins/VirtualProduction/CaptureManager/CaptureManagerCore/Source/CaptureSourceFramework/Private/CaptureSourceCapabilityEvent.cpp
deleted file mode 100644
diff --git a/Engine/Plugins/VirtualProduction/CaptureManager/CaptureManagerCore/Source/CaptureSourceFramework/Private/CaptureSourceCapabilityProperty.cpp b/Engine/Plugins/VirtualProduction/CaptureManager/CaptureManagerCore/Source/CaptureSourceFramework/Private/CaptureSourceCapabilityProperty.cpp
deleted file mode 100644
diff --git a/Engine/Plugins/VirtualProduction/CaptureManager/CaptureManagerCore/Source/CaptureSourceFramework/Private/CaptureSourceDetailsCustomization.cpp b/Engine/Plugins/VirtualProduction/CaptureManager/CaptureManagerCore/Source/CaptureSourceFramework/Private/CaptureSourceDetailsCustomization.cpp
deleted file mode 100644
diff --git a/Engine/Plugins/VirtualProduction/CaptureManager/CaptureManagerCore/Source/CaptureSourceFramework/Private/CaptureSourceError.cpp b/Engine/Plugins/VirtualProduction/CaptureManager/CaptureManagerCore/Source/CaptureSourceFramework/Private/CaptureSourceError.cpp
deleted file mode 100644
diff --git a/Engine/Plugins/VirtualProduction/CaptureManager/CaptureManagerCore/Source/CaptureSourceFramework/Private/CaptureSourceFactory.cpp b/Engine/Plugins/VirtualProduction/CaptureManager/CaptureManagerCore/Source/CaptureSourceFramework/Private/CaptureSourceFactory.cpp
deleted file mode 100644
diff --git a/Engine/Plugins/VirtualProduction/CaptureManager/CaptureManagerCore/Source/CaptureSourceFramework/Private/CaptureSourceFrameworkLog.h b/Engine/Plugins/VirtualProduction/CaptureManager/CaptureManagerCore/Source/CaptureSourceFramework/Private/CaptureSourceFrameworkLog.h
deleted file mode 100644
diff --git a/Engine/Plugins/VirtualProduction/CaptureManager/CaptureManagerCore/Source/CaptureSourceFramework/Private/CaptureSourceFrameworkModule.cpp b/Engine/Plugins/VirtualProduction/CaptureManager/CaptureManagerCore/Source/CaptureSourceFramework/Private/CaptureSourceFrameworkModule.cpp
deleted file mode 100644
diff --git a/Engine/Plugins/VirtualProduction/CaptureManager/CaptureManagerCore/Source/CaptureSourceFramework/Private/CaptureSourceHandle.cpp b/Engine/Plugins/VirtualProduction/CaptureManager/CaptureManagerCore/Source/CaptureSourceFramework/Private/CaptureSourceHandle.cpp
deleted file mode 100644
diff --git a/Engine/Plugins/VirtualProduction/CaptureManager/CaptureManagerCore/Source/CaptureSourceFramework/Private/CaptureSourceManager.cpp b/Engine/Plugins/VirtualProduction/CaptureManager/CaptureManagerCore/Source/CaptureSourceFramework/Private/CaptureSourceManager.cpp
deleted file mode 100644
diff --git a/Engine/Plugins/VirtualProduction/CaptureManager/CaptureManagerCore/Source/CaptureSourceFramework/Private/Tests/TestCaptureSource.cpp b/Engine/Plugins/VirtualProduction/CaptureManager/CaptureManagerCore/Source/CaptureSourceFramework/Private/Tests/TestCaptureSource.cpp
deleted file mode 100644
diff --git a/Engine/Plugins/VirtualProduction/CaptureManager/CaptureManagerCore/Source/CaptureSourceFramework/Private/Tests/TestCaptureSourceCapability.cpp b/Engine/Plugins/VirtualProduction/CaptureManager/CaptureManagerCore/Source/CaptureSourceFramework/Private/Tests/TestCaptureSourceCapability.cpp
deleted file mode 100644
diff --git a/Engine/Plugins/VirtualProduction/CaptureManager/CaptureManagerCore/Source/CaptureSourceFramework/Private/Tests/TestCaptureSourceFactory.cpp b/Engine/Plugins/VirtualProduction/CaptureManager/CaptureManagerCore/Source/CaptureSourceFramework/Private/Tests/TestCaptureSourceFactory.cpp
deleted file mode 100644
diff --git a/Engine/Plugins/VirtualProduction/CaptureManager/CaptureManagerCore/Source/CaptureSourceFramework/Private/Tests/TestCaptureSourceManager.cpp b/Engine/Plugins/VirtualProduction/CaptureManager/CaptureManagerCore/Source/CaptureSourceFramework/Private/Tests/TestCaptureSourceManager.cpp
deleted file mode 100644
diff --git a/Engine/Plugins/VirtualProduction/CaptureManager/CaptureManagerCore/Source/CaptureSourceFramework/Private/Widgets/SDefaultCaptureSourceCapabilityCreator.cpp b/Engine/Plugins/VirtualProduction/CaptureManager/CaptureManagerCore/Source/CaptureSourceFramework/Private/Widgets/SDefaultCaptureSourceCapabilityCreator.cpp
deleted file mode 100644
diff --git a/Engine/Plugins/VirtualProduction/CaptureManager/CaptureManagerCore/Source/CaptureSourceFramework/Private/Widgets/SDefaultCaptureSourceCommandCreator.cpp b/Engine/Plugins/VirtualProduction/CaptureManager/CaptureManagerCore/Source/CaptureSourceFramework/Private/Widgets/SDefaultCaptureSourceCommandCreator.cpp
deleted file mode 100644
diff --git a/Engine/Plugins/VirtualProduction/CaptureManager/CaptureManagerCore/Source/CaptureSourceFramework/Private/Widgets/SDefaultCaptureSourceCreator.cpp b/Engine/Plugins/VirtualProduction/CaptureManager/CaptureManagerCore/Source/CaptureSourceFramework/Private/Widgets/SDefaultCaptureSourceCreator.cpp
deleted file mode 100644
diff --git a/Engine/Plugins/VirtualProduction/CaptureManager/CaptureManagerCore/Source/CaptureSourceFramework/Private/Widgets/SDefaultCaptureSourcePropertyCreator.cpp b/Engine/Plugins/VirtualProduction/CaptureManager/CaptureManagerCore/Source/CaptureSourceFramework/Private/Widgets/SDefaultCaptureSourcePropertyCreator.cpp
deleted file mode 100644
diff --git a/Engine/Plugins/VirtualProduction/CaptureManager/CaptureManagerCore/Source/CaptureSourceFramework/Private/Widgets/WidgetUtils.cpp b/Engine/Plugins/VirtualProduction/CaptureManager/CaptureManagerCore/Source/CaptureSourceFramework/Private/Widgets/WidgetUtils.cpp
deleted file mode 100644
diff --git a/Engine/Plugins/VirtualProduction/CaptureManager/CaptureManagerCore/Source/CaptureSourceFramework/Private/Widgets/WidgetUtils.h b/Engine/Plugins/VirtualProduction/CaptureManager/CaptureManagerCore/Source/CaptureSourceFramework/Private/Widgets/WidgetUtils.h
deleted file mode 100644
diff --git a/Engine/Plugins/VirtualProduction/CaptureManager/CaptureManagerCore/Source/CaptureSourceFramework/Public/Capabilities/DiskCapacityCapability.h b/Engine/Plugins/VirtualProduction/CaptureManager/CaptureManagerCore/Source/CaptureSourceFramework/Public/Capabilities/DiskCapacityCapability.h
deleted file mode 100644
diff --git a/Engine/Plugins/VirtualProduction/CaptureManager/CaptureManagerCore/Source/CaptureSourceFramework/Public/Capabilities/IngestCapability.h b/Engine/Plugins/VirtualProduction/CaptureManager/CaptureManagerCore/Source/CaptureSourceFramework/Public/Capabilities/IngestCapability.h
deleted file mode 100644
diff --git a/Engine/Plugins/VirtualProduction/CaptureManager/CaptureManagerCore/Source/CaptureSourceFramework/Public/Capabilities/RecordingCapability.h b/Engine/Plugins/VirtualProduction/CaptureManager/CaptureManagerCore/Source/CaptureSourceFramework/Public/Capabilities/RecordingCapability.h
deleted file mode 100644
diff --git a/Engine/Plugins/VirtualProduction/CaptureManager/CaptureManagerCore/Source/CaptureSourceFramework/Public/CaptureSource.h b/Engine/Plugins/VirtualProduction/CaptureManager/CaptureManagerCore/Source/CaptureSourceFramework/Public/CaptureSource.h
deleted file mode 100644
diff --git a/Engine/Plugins/VirtualProduction/CaptureManager/CaptureManagerCore/Source/CaptureSourceFramework/Public/CaptureSourceCapability.h b/Engine/Plugins/VirtualProduction/CaptureManager/CaptureManagerCore/Source/CaptureSourceFramework/Public/CaptureSourceCapability.h
deleted file mode 100644
diff --git a/Engine/Plugins/VirtualProduction/CaptureManager/CaptureManagerCore/Source/CaptureSourceFramework/Public/CaptureSourceCapabilityCommand.h b/Engine/Plugins/VirtualProduction/CaptureManager/CaptureManagerCore/Source/CaptureSourceFramework/Public/CaptureSourceCapabilityCommand.h
deleted file mode 100644
diff --git a/Engine/Plugins/VirtualProduction/CaptureManager/CaptureManagerCore/Source/CaptureSourceFramework/Public/CaptureSourceCapabilityEvent.h b/Engine/Plugins/VirtualProduction/CaptureManager/CaptureManagerCore/Source/CaptureSourceFramework/Public/CaptureSourceCapabilityEvent.h
deleted file mode 100644
diff --git a/Engine/Plugins/VirtualProduction/CaptureManager/CaptureManagerCore/Source/CaptureSourceFramework/Public/CaptureSourceCapabilityProperty.h b/Engine/Plugins/VirtualProduction/CaptureManager/CaptureManagerCore/Source/CaptureSourceFramework/Public/CaptureSourceCapabilityProperty.h
deleted file mode 100644
diff --git a/Engine/Plugins/VirtualProduction/CaptureManager/CaptureManagerCore/Source/CaptureSourceFramework/Public/CaptureSourceDetailsCustomization.h b/Engine/Plugins/VirtualProduction/CaptureManager/CaptureManagerCore/Source/CaptureSourceFramework/Public/CaptureSourceDetailsCustomization.h
deleted file mode 100644
diff --git a/Engine/Plugins/VirtualProduction/CaptureManager/CaptureManagerCore/Source/CaptureSourceFramework/Public/CaptureSourceError.h b/Engine/Plugins/VirtualProduction/CaptureManager/CaptureManagerCore/Source/CaptureSourceFramework/Public/CaptureSourceError.h
deleted file mode 100644
diff --git a/Engine/Plugins/VirtualProduction/CaptureManager/CaptureManagerCore/Source/CaptureSourceFramework/Public/CaptureSourceFactory.h b/Engine/Plugins/VirtualProduction/CaptureManager/CaptureManagerCore/Source/CaptureSourceFramework/Public/CaptureSourceFactory.h
deleted file mode 100644
diff --git a/Engine/Plugins/VirtualProduction/CaptureManager/CaptureManagerCore/Source/CaptureSourceFramework/Public/CaptureSourceFrameworkModule.h b/Engine/Plugins/VirtualProduction/CaptureManager/CaptureManagerCore/Source/CaptureSourceFramework/Public/CaptureSourceFrameworkModule.h
deleted file mode 100644
diff --git a/Engine/Plugins/VirtualProduction/CaptureManager/CaptureManagerCore/Source/CaptureSourceFramework/Public/CaptureSourceHandle.h b/Engine/Plugins/VirtualProduction/CaptureManager/CaptureManagerCore/Source/CaptureSourceFramework/Public/CaptureSourceHandle.h
deleted file mode 100644
diff --git a/Engine/Plugins/VirtualProduction/CaptureManager/CaptureManagerCore/Source/CaptureSourceFramework/Public/CaptureSourceManager.h b/Engine/Plugins/VirtualProduction/CaptureManager/CaptureManagerCore/Source/CaptureSourceFramework/Public/CaptureSourceManager.h
deleted file mode 100644
diff --git a/Engine/Plugins/VirtualProduction/CaptureManager/CaptureManagerCore/Source/CaptureSourceFramework/Public/Widgets/SDefaultCaptureSourceCapabilityCreator.h b/Engine/Plugins/VirtualProduction/CaptureManager/CaptureManagerCore/Source/CaptureSourceFramework/Public/Widgets/SDefaultCaptureSourceCapabilityCreator.h
deleted file mode 100644
diff --git a/Engine/Plugins/VirtualProduction/CaptureManager/CaptureManagerCore/Source/CaptureSourceFramework/Public/Widgets/SDefaultCaptureSourceCommandCreator.h b/Engine/Plugins/VirtualProduction/CaptureManager/CaptureManagerCore/Source/CaptureSourceFramework/Public/Widgets/SDefaultCaptureSourceCommandCreator.h
deleted file mode 100644
diff --git a/Engine/Plugins/VirtualProduction/CaptureManager/CaptureManagerCore/Source/CaptureSourceFramework/Public/Widgets/SDefaultCaptureSourceCreator.h b/Engine/Plugins/VirtualProduction/CaptureManager/CaptureManagerCore/Source/CaptureSourceFramework/Public/Widgets/SDefaultCaptureSourceCreator.h
deleted file mode 100644
diff --git a/Engine/Plugins/VirtualProduction/CaptureManager/CaptureManagerCore/Source/CaptureSourceFramework/Public/Widgets/SDefaultCaptureSourcePropertyCreator.h b/Engine/Plugins/VirtualProduction/CaptureManager/CaptureManagerCore/Source/CaptureSourceFramework/Public/Widgets/SDefaultCaptureSourcePropertyCreator.h
deleted file mode 100644
diff --git a/Engine/Plugins/VirtualProduction/CaptureManager/CaptureManagerCore/Source/DataIngestCore/DataIngestCore.Build.cs b/Engine/Plugins/VirtualProduction/CaptureManager/CaptureManagerCore/Source/DataIngestCore/DataIngestCore.Build.cs
--- a/Engine/Plugins/VirtualProduction/CaptureManager/CaptureManagerCore/Source/DataIngestCore/DataIngestCore.Build.cs
+++ b/Engine/Plugins/VirtualProduction/CaptureManager/CaptureManagerCore/Source/DataIngestCore/DataIngestCore.Build.cs
@@ -19,7 +19,6 @@ public class DataIngestCore : ModuleRules
 
 		PrivateDependencyModuleNames.AddRange(new string[]
 		{
-			"CaptureSourceFramework",
 			"CoreUObject",
 			"InputCore",
 			"PropertyEditor",
diff --git a/Engine/Plugins/VirtualProduction/CaptureManager/CaptureManagerCore/Source/LiveLinkHubCaptureMessaging/LiveLinkHubCaptureMessaging.Build.cs b/Engine/Plugins/VirtualProduction/CaptureManager/CaptureManagerCore/Source/LiveLinkHubCaptureMessaging/LiveLinkHubCaptureMessaging.Build.cs
--- a/Engine/Plugins/VirtualProduction/CaptureManager/CaptureManagerCore/Source/LiveLinkHubCaptureMessaging/LiveLinkHubCaptureMessaging.Build.cs
+++ b/Engine/Plugins/VirtualProduction/CaptureManager/CaptureManagerCore/Source/LiveLinkHubCaptureMessaging/LiveLinkHubCaptureMessaging.Build.cs
@@ -20,7 +20,6 @@ public class LiveLinkHubCaptureMessaging : ModuleRules
 
 		PrivateDependencyModuleNames.AddRange(new string[]
 		{
-			"CaptureSourceFramework",
 			"MessagingCommon"
 		});
 	}
diff --git a/Engine/Plugins/VirtualProduction/CaptureManager/CaptureManagerEditor/Source/LiveLinkHubDiscoveryEditor/LiveLinkHubDiscoveryEditor.Build.cs b/Engine/Plugins/VirtualProduction/CaptureManager/CaptureManagerEditor/Source/LiveLinkHubDiscoveryEditor/LiveLinkHubDiscoveryEditor.Build.cs
--- a/Engine/Plugins/VirtualProduction/CaptureManager/CaptureManagerEditor/Source/LiveLinkHubDiscoveryEditor/LiveLinkHubDiscoveryEditor.Build.cs
+++ b/Engine/Plugins/VirtualProduction/CaptureManager/CaptureManagerEditor/Source/LiveLinkHubDiscoveryEditor/LiveLinkHubDiscoveryEditor.Build.cs
@@ -8,7 +8,6 @@ public class LiveLinkHubDiscoveryEditor : ModuleRules
 	{
 		PrivateDependencyModuleNames.AddRange(new string[] 
 		{
-			"CaptureSourceFramework",
 			"CaptureUtils",
 			"Core",
 			"CoreUObject",
diff --git a/Engine/Plugins/VirtualProduction/CaptureManager/CaptureManagerEditor/Source/LiveLinkHubExportServer/LiveLinkHubExportServer.Build.cs b/Engine/Plugins/VirtualProduction/CaptureManager/CaptureManagerEditor/Source/LiveLinkHubExportServer/LiveLinkHubExportServer.Build.cs
--- a/Engine/Plugins/VirtualProduction/CaptureManager/CaptureManagerEditor/Source/LiveLinkHubExportServer/LiveLinkHubExportServer.Build.cs
+++ b/Engine/Plugins/VirtualProduction/CaptureManager/CaptureManagerEditor/Source/LiveLinkHubExportServer/LiveLinkHubExportServer.Build.cs
@@ -20,7 +20,6 @@ public class LiveLinkHubExportServer : ModuleRules
 		PrivateDependencyModuleNames.AddRange(new string[]
 		{
 			"Sockets",
-			"CaptureSourceFramework",
 			"LiveLinkHubCaptureMessaging"
 		});
 	}
