Implementation of base ingest device class

diff --git a/Engine/Plugins/VirtualProduction/CaptureManager/CaptureManagerApp/CaptureManagerApp.uplugin b/Engine/Plugins/VirtualProduction/CaptureManager/CaptureManagerApp/CaptureManagerApp.uplugin
--- a/Engine/Plugins/VirtualProduction/CaptureManager/CaptureManagerApp/CaptureManagerApp.uplugin
+++ b/Engine/Plugins/VirtualProduction/CaptureManager/CaptureManagerApp/CaptureManagerApp.uplugin
@@ -76,6 +76,11 @@
 			"Name": "LiveLinkCapabilities",
 			"Type": "Editor",
 			"LoadingPhase": "Default"
+		},
+		{
+			"Name": "IngestLiveLinkDevice",
+			"Type": "Editor",
+			"LoadingPhase": "Default"
 		}
 	],
 	"Plugins": [
diff --git a/Engine/Plugins/VirtualProduction/CaptureManager/CaptureManagerApp/Source/CaptureManagerEditor/Private/Scripting/ULiveLinkHubCaptureDevice.cpp b/Engine/Plugins/VirtualProduction/CaptureManager/CaptureManagerApp/Source/CaptureManagerEditor/Private/Scripting/ULiveLinkHubCaptureDevice.cpp
--- a/Engine/Plugins/VirtualProduction/CaptureManager/CaptureManagerApp/Source/CaptureManagerEditor/Private/Scripting/ULiveLinkHubCaptureDevice.cpp
+++ b/Engine/Plugins/VirtualProduction/CaptureManager/CaptureManagerApp/Source/CaptureManagerEditor/Private/Scripting/ULiveLinkHubCaptureDevice.cpp
@@ -281,7 +281,7 @@ TObjectPtr<UIngestCapability_ProcessResult> ULiveLinkHubCaptureDevice::FImpl::Fe
 
 	UpdateTakeListCallback->Callback = MoveTemp(OnTakesListUpdated);
 
-	ILiveLinkDeviceCapability_Ingest::Execute_UpdateTakeListAsync(BackgroundDevice);
+	ILiveLinkDeviceCapability_Ingest::Execute_UpdateTakeList(BackgroundDevice);
 	TakeListUpdatedEvent->Wait();
 
 	OutTakes = MoveTemp(Takes);
diff --git a/Engine/Plugins/VirtualProduction/CaptureManager/CaptureManagerApp/Source/CaptureManagerEditor/Private/UI/TakesView.cpp b/Engine/Plugins/VirtualProduction/CaptureManager/CaptureManagerApp/Source/CaptureManagerEditor/Private/UI/TakesView.cpp
--- a/Engine/Plugins/VirtualProduction/CaptureManager/CaptureManagerApp/Source/CaptureManagerEditor/Private/UI/TakesView.cpp
+++ b/Engine/Plugins/VirtualProduction/CaptureManager/CaptureManagerApp/Source/CaptureManagerEditor/Private/UI/TakesView.cpp
@@ -534,7 +534,7 @@ void FTakesView::CaptureDeviceStarted(FGuid InCaptureDeviceId)
 		}));
 		
 		ILiveLinkDeviceCapability_Ingest::Execute_CreateUpdateTakeListCallback(Device)->Callback = MoveTemp(Callback);
-		ILiveLinkDeviceCapability_Ingest::Execute_UpdateTakeListAsync(Device);
+		ILiveLinkDeviceCapability_Ingest::Execute_UpdateTakeList(Device);
 	});
 }
 
diff --git a/Engine/Plugins/VirtualProduction/CaptureManager/CaptureManagerApp/Source/GenericLiveLinkDevice/GenericLiveLinkDevice.Build.cs b/Engine/Plugins/VirtualProduction/CaptureManager/CaptureManagerApp/Source/GenericLiveLinkDevice/GenericLiveLinkDevice.Build.cs
--- a/Engine/Plugins/VirtualProduction/CaptureManager/CaptureManagerApp/Source/GenericLiveLinkDevice/GenericLiveLinkDevice.Build.cs
+++ b/Engine/Plugins/VirtualProduction/CaptureManager/CaptureManagerApp/Source/GenericLiveLinkDevice/GenericLiveLinkDevice.Build.cs
@@ -19,16 +19,13 @@ public class GenericLiveLinkDevice : ModuleRules
 		PrivateDependencyModuleNames.AddRange(new string[]
 		{
 			"CaptureProtocolStack",
-			"CaptureDataConverter",
-			"CaptureCommunicationManager",
 			"Engine",
 			"Json",
+			"IngestLiveLinkDevice",
 			"LiveLinkDevice",
 			"LiveLinkCapabilities",
 			"Media",
 			"LiveLinkFaceMetadata",
-			"LiveLinkHubCaptureMessaging",
-			"LiveLinkHubExportClient",
 			"LiveLinkHub",
 			"ToolWidgets",
 			"Slate",
diff --git a/Engine/Plugins/VirtualProduction/CaptureManager/CaptureManagerApp/Source/GenericLiveLinkDevice/Private/Archive/GenericLiveLinkArchiveDevice.cpp b/Engine/Plugins/VirtualProduction/CaptureManager/CaptureManagerApp/Source/GenericLiveLinkDevice/Private/Archive/GenericLiveLinkArchiveDevice.cpp
--- a/Engine/Plugins/VirtualProduction/CaptureManager/CaptureManagerApp/Source/GenericLiveLinkDevice/Private/Archive/GenericLiveLinkArchiveDevice.cpp
+++ b/Engine/Plugins/VirtualProduction/CaptureManager/CaptureManagerApp/Source/GenericLiveLinkDevice/Private/Archive/GenericLiveLinkArchiveDevice.cpp
@@ -3,15 +3,15 @@
 #include "GenericLiveLinkArchiveDevice.h"
 
 #include "HAL/FileManager.h"
-#include "CaptureCommunicationManagerModule.h"
 #include "Settings/CaptureManagerSettings.h"
 
-#include "Utils/GenericLiveLinkDeviceUtility.h"
 #include "Utils/GenericLiveLinkDeviceLog.h"
 
 #include "LiveLinkFaceMetadata.h"
 #include "StereoCameraMetadataParseUtils.h"
 
+#include "Utils/GenericLiveLinkDeviceUtility.h"
+
 DEFINE_LOG_CATEGORY(LogGenericLiveLinkArchiveDevice);
 
 const UGenericLiveLinkArchiveDeviceSettings* UGenericLiveLinkArchiveDevice::GetSettings() const
@@ -39,15 +39,6 @@ FText UGenericLiveLinkArchiveDevice::GetHealthText() const
 	return FText::FromString("Nominal");
 }
 
-void UGenericLiveLinkArchiveDevice::OnDeviceAdded(const FGuid InDeviceGuid, ULiveLinkDeviceSettings* InSettings)
-{
-	Super::OnDeviceAdded(InDeviceGuid, InSettings);
-
-	DeviceGuid = InDeviceGuid;
-
-	CommunicationManager = UE::CaptureManager::Private::GetCommunicationManager();
-}
-
 void UGenericLiveLinkArchiveDevice::OnSettingChanged(const FPropertyChangedEvent& InPropertyChangedEvent)
 {
 	Super::OnSettingChanged(InPropertyChangedEvent);
@@ -66,7 +57,38 @@ void UGenericLiveLinkArchiveDevice::OnSettingChanged(const FPropertyChangedEvent
 	}
 }
 
-void UGenericLiveLinkArchiveDevice::UpdateTakeListAsync_Implementation()
+FString UGenericLiveLinkArchiveDevice::GetFullTakePath(UE::CaptureManager::FTakeId InTakeId) const
+{
+	FString StoragePath = GetSettings()->StoragePath.Path;
+
+	if (const FString* RelativeTakePath = RelativeTakePaths.Find(InTakeId))
+	{
+		return FPaths::Combine(StoragePath, *RelativeTakePath);
+	}
+
+	return FString();
+}
+
+void UGenericLiveLinkArchiveDevice::RunIngestTake(int32 InTakeId, const UIngestCapability_Options* InIngestOptions)
+{
+	AsyncTask(ENamedThreads::AnyThread, [this,
+			  TakeId = InTakeId,
+			  IngestOptions = InIngestOptions]()
+	{
+		static constexpr uint32 NumberOfTasks = 2; // Convert, Upload
+
+		using namespace UE::CaptureManager;
+		TSharedPtr<FTaskProgress> TaskProgress
+			= MakeShared<FTaskProgress>(NumberOfTasks, FTaskProgress::FProgressReporter::CreateLambda([this, TakeId](double InProgress)
+		{
+			ExecuteProcessProgressReporter(TakeId, InProgress);
+		}));
+
+		Super::IngestTake(TakeId, IngestOptions, MoveTemp(TaskProgress));
+	});
+}
+
+void UGenericLiveLinkArchiveDevice::UpdateTakeList_Implementation()
 {
 	AsyncTask(ENamedThreads::AnyThread, [this]()
 	{
@@ -112,147 +134,6 @@ void UGenericLiveLinkArchiveDevice::UpdateTakeListAsync_Implementation()
 	});
 }
 
-void UGenericLiveLinkArchiveDevice::RunIngestTakeAsync(int32 InTakeId, const UIngestCapability_Options* InIngestOptions)
-{
-	AsyncTask(ENamedThreads::AnyThread, [this, 
-			  TakeId = InTakeId, 
-			  IngestOptions = TStrongObjectPtr<const UIngestCapability_Options>(InIngestOptions)]()
-	{
-		using namespace UE::CaptureManager;
-
-		TOptional<FTakeMetadata> MaybeTakeMetadata = GetTakeMetadata(TakeId);
-		if (!MaybeTakeMetadata.IsSet())
-		{
-			FString Message = TEXT("Failed to find the take metadata for the specified take");
-			ExecuteProcessFinishedReporter(TakeId, MakeError(FIngestCapability_Error::InvalidArgument, MoveTemp(Message)));
-			return;
-		}
-
-		FTakeMetadata TakeMetadata = MoveTemp(MaybeTakeMetadata.GetValue());
-
-		FString RelativeTakePath = RelativeTakePaths[TakeId];
-
-		static constexpr uint32 NumberOfTasks = 2; // Convert, Upload
-
-		TSharedPtr<FTaskProgress> TaskProgress 
-			= MakeShared<FTaskProgress>(NumberOfTasks, FTaskProgress::FProgressReporter::CreateLambda([this, TakeId](double InProgress)
-		{
-			ExecuteProcessProgressReporter(TakeId, InProgress);
-		}));
-
-		FString TakeName = TakeMetadata.Slate + TEXT("_") + FString::FromInt(TakeMetadata.TakeNumber);
-
-		const FString TakeConversionDirectory = IngestOptions->WorkingDirectory;
-		const FString StoragePath = GetSettings()->StoragePath.Path;
-
-		FGuid TakeUploadId = FGuid::NewGuid();
-
-		FCaptureDataConverterParams Params;
-		Params.TakeName = TakeName;
-		Params.TakeMetadata = TakeMetadata;
-		Params.TakeOriginDirectory = FPaths::Combine(StoragePath, RelativeTakePath);
-		Params.TakeOutputDirectory = FPaths::Combine(TakeConversionDirectory, TakeName + TEXT("_") + TakeUploadId.ToString());
-
-		FCaptureDataConverterResult<void> Result = RunConversion(TakeId, IngestOptions, TaskProgress, Params);
-
-		Converters.Remove(TakeId);
-
-		if (Result.HasValue())
-		{
-			if (!Uploader.IsValid())
-			{
-				Uploader = CommunicationManager->GetUploader(DeviceGuid,
-															 GetName(),
-															 FUploader::FUploadStateUpdate::CreateUObject(this, &UGenericLiveLinkArchiveDevice::UploadTakeStateHandler),
-															 FUploader::FUploadFinished::CreateUObject(this, &UGenericLiveLinkArchiveDevice::UploadTakeFinished));
-			}
-
-			if (TSharedPtr<FUploader> SharedUploader = Uploader.Pin())
-			{
-				FScopeLock Lock(&Mutex);
-
-				FTaskProgress::FTask UploadTask = TaskProgress->StartTask();
-
-				ImportContexts.Add(TakeUploadId, { TakeName, TakeId, IngestOptions, TaskProgress, UploadTask });
-
-				int32 Worker = SharedUploader->UploadTake(TakeUploadId, Params.TakeOutputDirectory, MoveTemp(TakeMetadata));
-				WorkerForTake.Add(TakeId, Worker);
-			}
-			else
-			{
-				IFileManager::Get().DeleteDirectory(*Params.TakeOutputDirectory, false, true);
-
-				FString Message = TEXT("Failed to upload take: Uploader not found");
-				ExecuteProcessFinishedReporter(TakeId, MakeError(FIngestCapability_Error::UploaderNotFound, MoveTemp(Message)));
-			}
-		}
-		else
-		{
-			IFileManager::Get().DeleteDirectory(*Params.TakeOutputDirectory, false, true);
-
-			FCaptureDataConverterError Error = Result.StealError();
-
-			FText Message = FText::Join(FText::FromString(TEXT("\n\r")), Error.GetErrors());
-			ExecuteProcessFinishedReporter(TakeId, 
-										   MakeError(FIngestCapability_Error::ConversionError, Message.ToString()));
-		}
-	});
-}
-
-void UGenericLiveLinkArchiveDevice::CancelIngest(int32 InTakeId)
-{
-	if (TUniquePtr<FCaptureDataConverter>* Converter = Converters.Find(InTakeId))
-	{
-		(*Converter)->Cancel();
-	}
-
-	int32 WorkerId = INDEX_NONE;
-
-	{
-		FScopeLock Lock(&Mutex);
-		if (int32* WorkerIdElem = WorkerForTake.Find(InTakeId))
-		{
-			WorkerId = *WorkerIdElem;
-		}
-	}
-
-	if (WorkerId != INDEX_NONE)
-	{
-		if (TSharedPtr<FUploader> SharedUploader = Uploader.Pin())
-		{
-			SharedUploader->AbortUpload(WorkerId);
-		}
-		else
-		{
-			// If `UploadTakeFinished` is not called (eg. connection between LLH and UE broke), `Uploader` will be deleted, but UI will stay in state where import is in progress 99.9%.
-			// Clicking stop button at that moment will not enter if block above, as uploader does not exist, so in that case we have to manually cleanup the workspace and put UI in proper state.
-			FScopeLock Lock(&Mutex);
-
-			TMap<FGuid, FContext> FilteredMap = ImportContexts.FilterByPredicate([InTakeId](const TPair<FGuid, FContext>& Item) -> bool
-			{
-				return Item.Value.TakeId == InTakeId;
-			});
-
-			check(FilteredMap.Num() == 1);
-			TMap<FGuid, FContext>::TConstIterator FoundContext = FilteredMap.CreateConstIterator();
-			check(FoundContext);
-
-			const FContext& Context = FoundContext->Value;
-
-			const FString TakeDirectory = FPaths::Combine(Context.IngestOptions->WorkingDirectory,
-														  Context.TakeName + TEXT("_") + FoundContext->Key.ToString());
-			IFileManager::Get().DeleteDirectory(*TakeDirectory, false, true);
-
-			FString Message = TEXT("Aborted by user");
-			ExecuteProcessFinishedReporter(InTakeId, MakeError(FIngestCapability_Error::AbortedByUser, MoveTemp(Message)));
-
-			ImportContexts.Remove(FoundContext->Key);
-
-			WorkerForTake.Remove(InTakeId);
-		}
-	}
-}
-
 ELiveLinkDeviceConnectionStatus UGenericLiveLinkArchiveDevice::GetConnectionStatus_Implementation() const
 {
 	const UGenericLiveLinkArchiveDeviceSettings* DeviceSettings = GetSettings();
@@ -296,97 +177,6 @@ bool UGenericLiveLinkArchiveDevice::Disconnect_Implementation()
 	return true;
 }
 
-FCaptureDataConverterResult<void> UGenericLiveLinkArchiveDevice::RunConversion(int32 InTakeId,
-																			   TStrongObjectPtr<const UIngestCapability_Options> InIngestOptions,
-																			   TSharedPtr<UE::CaptureManager::FTaskProgress> InTaskProgress,
-																			   FCaptureDataConverterParams& InParams)
-{
-	using namespace UE::CaptureManager;
-
-	FTaskProgress::FTask ConvertTask = InTaskProgress->StartTask();
-
-	FCaptureConvertVideoOutputParams VideoParams;
-	VideoParams.Format = InIngestOptions->Video.Format;
-	VideoParams.ImageFileName = InIngestOptions->Video.FileNamePrefix;
-	VideoParams.OutputPixelFormat = Private::ConvertPixelFormat(InIngestOptions->Video.PixelFormat);
-
-	InParams.VideoOutputParams = MoveTemp(VideoParams);
-
-	FCaptureConvertAudioOutputParams AudioParams;
-	AudioParams.Format = InIngestOptions->Audio.Format;
-	AudioParams.AudioFileName = InIngestOptions->Audio.FileNamePrefix;
-
-	InParams.AudioOutputParams = MoveTemp(AudioParams);
-
-	FCaptureConvertDepthOutputParams DepthParams;
-	DepthParams.ImageFileName = TEXT("depth");
-
-	InParams.DepthOutputParams = MoveTemp(DepthParams);
-
-	FCaptureConvertCalibrationOutputParams CalibrationParams;
-	CalibrationParams.FileName = TEXT("calibration");
-
-	InParams.CalibrationOutputParams = MoveTemp(CalibrationParams);
-
-	FCaptureDataConverter::FProgressReporter ProgressReporter =
-		FCaptureDataConverter::FProgressReporter::CreateLambda([ConvertTask](const double InProgress) mutable
-	{
-		ConvertTask.Update(InProgress);
-	});
-
-	TUniquePtr<FCaptureDataConverter>& Converter = Converters.FindOrAdd(InTakeId, MakeUnique<FCaptureDataConverter>());
-	return Converter->Run(InParams, MoveTemp(ProgressReporter));
-}
-
-void UGenericLiveLinkArchiveDevice::UploadTakeStateHandler(const FGuid& InTakeUploadId, double InProgress)
-{
-	if (!ensure(InTakeUploadId.IsValid()))
-	{
-		return;
-	}
-
-	FScopeLock Lock(&Mutex);
-	if (FContext* Context = ImportContexts.Find(InTakeUploadId))
-	{
-		Context->Task.Update(InProgress);
-	}
-}
-
-void UGenericLiveLinkArchiveDevice::UploadTakeFinished(const FGuid& InTakeUploadId, const FString& InMessage, int32 InCode)
-{
-	using namespace UE::CaptureManager;
-
-	if (!ensure(InTakeUploadId.IsValid()))
-	{
-		return;
-	}
-
-	FScopeLock Lock(&Mutex);
-	if (const FContext* Context = ImportContexts.Find(InTakeUploadId))
-	{
-		const bool bShouldCleanConversionDirectory = GetDefault<UCaptureManagerSettings>()->bShouldCleanConversionDirectory;
-		if (bShouldCleanConversionDirectory)
-		{
-			const FString ConversionDirectory = FPaths::Combine(Context->IngestOptions->WorkingDirectory,
-																Context->TakeName + TEXT("_") + InTakeUploadId.ToString());
-			IFileManager::Get().DeleteDirectory(*ConversionDirectory, false, true);
-		}
-
-		if (InMessage.IsEmpty())
-		{
-			ExecuteProcessFinishedReporter(Context->TakeId, MakeValue());
-		}
-		else
-		{
-			ExecuteProcessFinishedReporter(Context->TakeId, MakeError(FIngestCapability_Error::UploaderError, InMessage));
-		}
-
-		ImportContexts.Remove(InTakeUploadId);
-
-		WorkerForTake.Remove(Context->TakeId);
-	}
-}
-
 TOptional<FTakeMetadata> UGenericLiveLinkArchiveDevice::ReadTake(const FString& InCurrentTakeFile)
 {
 	using namespace UE::CaptureManager;
diff --git a/Engine/Plugins/VirtualProduction/CaptureManager/CaptureManagerApp/Source/GenericLiveLinkDevice/Private/Archive/GenericLiveLinkArchiveDevice.h b/Engine/Plugins/VirtualProduction/CaptureManager/CaptureManagerApp/Source/GenericLiveLinkDevice/Private/Archive/GenericLiveLinkArchiveDevice.h
--- a/Engine/Plugins/VirtualProduction/CaptureManager/CaptureManagerApp/Source/GenericLiveLinkDevice/Private/Archive/GenericLiveLinkArchiveDevice.h
+++ b/Engine/Plugins/VirtualProduction/CaptureManager/CaptureManagerApp/Source/GenericLiveLinkDevice/Private/Archive/GenericLiveLinkArchiveDevice.h
@@ -10,11 +10,8 @@
 #include "Misc/Optional.h"
 
 #include "Async/TaskProgress.h"
-#include "LiveLinkHubExportManager.h"
 
-#include "CaptureDataConverter.h"
-
-#include "CaptureCommunicationManager.h"
+#include "BaseIngestLiveLinkDevice.h"
 
 #include "GenericLiveLinkArchiveDevice.generated.h"
 
@@ -34,8 +31,7 @@ public:
 
 
 UCLASS(BlueprintType)
-class UGenericLiveLinkArchiveDevice : public ULiveLinkDevice
-	, public ILiveLinkDeviceCapability_Ingest
+class UGenericLiveLinkArchiveDevice final : public UBaseIngestLiveLinkDevice
 	, public ILiveLinkDeviceCapability_Connection
 {
 	GENERATED_BODY()
@@ -50,16 +46,15 @@ public:
 	virtual EDeviceHealth GetDeviceHealth() const override;
 	virtual FText GetHealthText() const override;
 
-	virtual void OnDeviceAdded(const FGuid InDeviceGuid, ULiveLinkDeviceSettings* InSettings) override;
 	virtual void OnSettingChanged(const FPropertyChangedEvent& InPropertyChangedEvent) override;
 	//~ End ULiveLinkDevice interface
 
 private:
+	virtual FString GetFullTakePath(UE::CaptureManager::FTakeId InTakeId) const override;
+
 	//~ Begin ULiveLinkDeviceCapability_Ingest interface
-	virtual void UpdateTakeListAsync_Implementation() override;
-	virtual void RunIngestTakeAsync(int32 InTakeId,
-									const UIngestCapability_Options* InIngestOptions) override;
-	virtual void CancelIngest(int32 InTakeId) override;
+	virtual void RunIngestTake(int32 InTakeId, const UIngestCapability_Options* InIngestOptions) override;
+	virtual void UpdateTakeList_Implementation() override;
 	//~ End ULiveLinkDeviceCapability_Ingest interface
 
 	//~ Begin ILiveLinkDeviceCapability_Connection interface
@@ -72,34 +67,5 @@ private:
 
 	static TOptional<FTakeMetadata> ReadTake(const FString& InCurrentTakeFile);
 
-	FCaptureDataConverterResult<void> RunConversion(int32 InTakeId,
-													TStrongObjectPtr<const UIngestCapability_Options> InIngestOptions,
-													TSharedPtr<UE::CaptureManager::FTaskProgress> InTaskProgress,
-													FCaptureDataConverterParams& InParams);
-
-	void UploadTakeStateHandler(const FGuid& InTakeUploadId, double InProgress);
-	void UploadTakeFinished(const FGuid& InTakeUploadId, const FString& InMessage, int32 InCode);
-
-	struct FContext
-	{
-		FString TakeName;
-		int32 TakeId = INDEX_NONE;
-		TStrongObjectPtr<const UIngestCapability_Options> IngestOptions;
-		TSharedPtr<UE::CaptureManager::FTaskProgress> TaskProgress;
-		UE::CaptureManager::FTaskProgress::FTask Task;
-	};
-
-	TSharedPtr<FCommunicationManager> CommunicationManager;
-	TWeakPtr<FUploader> Uploader;
-
-	TMap<FGuid, FContext> ImportContexts;
 	TMap<int32, FString> RelativeTakePaths;
-
-	FCriticalSection Mutex;
-	TMap<int32, int32> WorkerForTake;
-
-	TMap<int32, TUniquePtr<FCaptureDataConverter>> Converters;
-	FTakeMetadataParser TakeMetadataParser;
-
-	FGuid DeviceGuid;
 };
diff --git a/Engine/Plugins/VirtualProduction/CaptureManager/CaptureManagerApp/Source/GenericLiveLinkDevice/Private/CPS/GenericLiveLinkCPSDevice.cpp b/Engine/Plugins/VirtualProduction/CaptureManager/CaptureManagerApp/Source/GenericLiveLinkDevice/Private/CPS/GenericLiveLinkCPSDevice.cpp
--- a/Engine/Plugins/VirtualProduction/CaptureManager/CaptureManagerApp/Source/GenericLiveLinkDevice/Private/CPS/GenericLiveLinkCPSDevice.cpp
+++ b/Engine/Plugins/VirtualProduction/CaptureManager/CaptureManagerApp/Source/GenericLiveLinkDevice/Private/CPS/GenericLiveLinkCPSDevice.cpp
@@ -2,7 +2,6 @@
 
 #include "GenericLiveLinkCPSDevice.h"
 
-#include "CaptureCommunicationManagerModule.h"
 #include "Settings/CaptureManagerSettings.h"
 
 #include "ILiveLinkRecordingSessionInfo.h"
@@ -16,6 +15,8 @@
 #include "LiveLinkFaceMetadata.h"
 #include "StereoCameraMetadataParseUtils.h"
 
+#include "HAL/FileManager.h"
+
 #include "Utils/GenericLiveLinkDeviceUtility.h"
 
 const UGenericLiveLinkCPSDeviceSettings* UGenericLiveLinkCPSDevice::GetSettings() const
@@ -47,28 +48,28 @@ void UGenericLiveLinkCPSDevice::OnDeviceAdded(const FGuid InDeviceGuid, ULiveLin
 {
 	CastChecked<UGenericLiveLinkCPSDeviceSettings>(InSettings)->ConnectAction.DeviceGuid = InDeviceGuid;
 
-	DeviceGuid = InDeviceGuid;
-
-	CommunicationManager = UE::CaptureManager::Private::GetCommunicationManager();
-
 	Super::OnDeviceAdded(InDeviceGuid, InSettings);
 }
 
 void UGenericLiveLinkCPSDevice::OnDeviceRemoved()
 {
-	RemoveAllTakes();
-
 	ILiveLinkDeviceCapability_Connection::Execute_Disconnect(this);
 
 	Super::OnDeviceRemoved();
 }
 
-void UGenericLiveLinkCPSDevice::OnSettingChanged(const FPropertyChangedEvent& InPropertyChangedEvent)
+FString UGenericLiveLinkCPSDevice::GetFullTakePath(UE::CaptureManager::FTakeId InTakeId) const
 {
-	Super::OnSettingChanged(InPropertyChangedEvent);
+	FScopeLock Lock(&DownloadedTakesMutex);
+	if (const FString* DownloadedTake = DownloadedTakes.Find(InTakeId))
+	{
+		return *DownloadedTake;
+	}
+	
+	return FString();
 }
 
-void UGenericLiveLinkCPSDevice::UpdateTakeListAsync_Implementation()
+void UGenericLiveLinkCPSDevice::UpdateTakeList_Implementation()
 {
 	using namespace UE::CaptureManager;
 
@@ -95,7 +96,7 @@ void UGenericLiveLinkCPSDevice::UpdateTakeListAsync_Implementation()
 	ExecuteUpdateTakeListCallback(Execute_GetTakeIdentifiers(this));
 }
 
-void UGenericLiveLinkCPSDevice::RunIngestTakeAsync(int32 InTakeId, const UIngestCapability_Options* InIngestOptions)
+void UGenericLiveLinkCPSDevice::RunIngestTake(int32 InTakeId, const UIngestCapability_Options* InIngestOptions)
 {
 	using namespace UE::CaptureManager;
 
@@ -132,56 +133,7 @@ void UGenericLiveLinkCPSDevice::CancelIngest(int32 InTakeId)
 {
 	Device->CancelExport(InTakeId);
 
-	if (TUniquePtr<FCaptureDataConverter>* Converter = Converters.Find(InTakeId))
-	{
-		(*Converter)->Cancel();
-	}
-
-	int32 WorkerId = INDEX_NONE;
-
-	{
-		FScopeLock Lock(&Mutex);
-		if (int32* WorkerIdElem = WorkerForTake.Find(InTakeId))
-		{
-			WorkerId = *WorkerIdElem;
-		}
-	}
-
-	if (WorkerId != INDEX_NONE)
-	{
-		if (TSharedPtr<FUploader> SharedUploader = Uploader.Pin())
-		{
-			SharedUploader->AbortUpload(WorkerId);
-		}
-		else
-		{
-			// If `UploadTakeFinished` is not called (eg. connection between LLH and UE broke), `Uploader` will be deleted, but UI will stay in state where import is in progress 99.9%.
-			// Clicking stop button at that moment will not enter if block above, as uploader does not exist, so in that case we have to manually cleanup the workspace and put UI in proper state.
-			FScopeLock Lock(&Mutex);
-
-			TMap<FGuid, FContext> FilteredMap = ImportContexts.FilterByPredicate([InTakeId](const TPair<FGuid, FContext>& Item) -> bool
-			{
-				return Item.Value.TakeId == InTakeId;
-			});
-
-			check(FilteredMap.Num() == 1);
-			TMap<FGuid, FContext>::TConstIterator FoundContext = FilteredMap.CreateConstIterator();
-			check(FoundContext);
-
-			const FContext& Context = FoundContext->Value;
-
-			const FString TakeDirectory = FPaths::Combine(Context.IngestOptions->WorkingDirectory,
-														  Context.TakeName + TEXT("_") + FoundContext->Key.ToString());
-			IFileManager::Get().DeleteDirectory(*TakeDirectory, false, true);
-
-			FString Message = TEXT("Aborted by user");
-			ExecuteProcessFinishedReporter(InTakeId, MakeError(FIngestCapability_Error::AbortedByUser, MoveTemp(Message)));
-
-			ImportContexts.Remove(FoundContext->Key);
-
-			WorkerForTake.Remove(InTakeId);
-		}
-	}
+	Super::CancelIngest(InTakeId);
 }
 
 ELiveLinkDeviceConnectionStatus UGenericLiveLinkCPSDevice::GetConnectionStatus_Implementation() const
@@ -429,181 +381,30 @@ void UGenericLiveLinkCPSDevice::OnExportFinished(UE::CaptureManager::TProtocolRe
 	if (InResult.HasValue())
 	{
 		// Blocking function
-		ConvertData(InTakeName, InTakeId, InTaskProgress, InIngestOptions);
-	}
-	else
-	{
-		TValueOrError<void, FIngestCapability_Error> Result = MakeError(FIngestCapability_Error::DownloaderError, InResult.GetError().GetMessage());
-
-		ExecuteProcessFinishedReporter(InTakeId, MoveTemp(Result));
-	}
-
-	IFileManager::Get().DeleteDirectory(*DownloadedTake, false, true);
-}
-
-void UGenericLiveLinkCPSDevice::ConvertData(const FString& InTakeName,
-											UE::CaptureManager::FTakeId InTakeId,
-											TSharedPtr<UE::CaptureManager::FTaskProgress> InTaskProgress,
-											TStrongObjectPtr<const UIngestCapability_Options> InIngestOptions)
-{
-	using namespace UE::CaptureManager;
-
-	check(InTaskProgress);
-	check(InIngestOptions);
-
-	const FString DownloadedStorage = GetDefault<UCaptureManagerSettings>()->DownloadDirectory.Path;
-
-	// Take metadata stored in the capability doesn't have enough information to convert the data
-	// so we need to parse the downloaded take.json
-	TOptional<FTakeMetadata> TakeMetadata = ParseTake(InTakeName);
-
-	if (!TakeMetadata.IsSet())
-	{
-		ExecuteProcessFinishedReporter(InTakeId, MakeError(FIngestCapability_Error::InternalError, TEXT("Failed to parse the take metadata")));
-		return;
-	}
-
-	const FString TakeConversionDirectory = InIngestOptions->WorkingDirectory;
-	FGuid TakeUploadId = FGuid::NewGuid();
-
-	FCaptureDataConverterParams Params;
-	Params.TakeName = InTakeName;
-	Params.TakeMetadata = MoveTemp(TakeMetadata.GetValue());
-	Params.TakeOriginDirectory = FPaths::Combine(DownloadedStorage, InTakeName);
-	Params.TakeOutputDirectory = FPaths::Combine(TakeConversionDirectory, InTakeName + TEXT("_") + TakeUploadId.ToString());
-
-	FCaptureDataConverterResult<void> Result = RunConversion(InTakeId, InIngestOptions, InTaskProgress, Params);
-
-	Converters.Remove(InTakeId);
-
-	if (Result.HasValue())
-	{
-		if (!Uploader.IsValid())
+		TOptional<FTakeMetadata> TakeMetadata = ParseTake(InTakeName);
+		if (TakeMetadata.IsSet())
 		{
-			Uploader = CommunicationManager->GetUploader(DeviceGuid,
-														 GetName(),
-														 FUploader::FUploadStateUpdate::CreateUObject(this, &UGenericLiveLinkCPSDevice::UploadTakeStateHandler),
-														 FUploader::FUploadFinished::CreateUObject(this, &UGenericLiveLinkCPSDevice::UploadTakeFinished));
-		}
-
-		if (TSharedPtr<FUploader> SharedUploader = Uploader.Pin())
-		{
-			FScopeLock Lock(&Mutex);
-			FTaskProgress::FTask UploadTask = InTaskProgress->StartTask();
-
-			ImportContexts.Add(TakeUploadId, { InTakeName, InTakeId, InIngestOptions, InTaskProgress, UploadTask });
+			FScopeLock Lock(&DownloadedTakesMutex);
+			DownloadedTakes.Add(InTakeId, DownloadedTake);
+			Lock.Unlock();
 
-			int32 Worker = SharedUploader->UploadTake(TakeUploadId, Params.TakeOutputDirectory, Params.TakeMetadata);
-
-			WorkerForTake.Add(InTakeId, Worker);
+			UpdateTake(InTakeId, MoveTemp(TakeMetadata.GetValue()));
+			Super::IngestTake(InTakeId, InIngestOptions.Get(), InTaskProgress);
 		}
 		else
 		{
-			IFileManager::Get().DeleteDirectory(*Params.TakeOutputDirectory, false, true);
-
-			FString Message = TEXT("Failed to upload take: Uploader not found");
-			ExecuteProcessFinishedReporter(InTakeId, MakeError(FIngestCapability_Error::UploaderNotFound, MoveTemp(Message)));
+			TValueOrError<void, FIngestCapability_Error> Result = MakeError(FIngestCapability_Error::DownloaderError, TEXT("Failed to parse the take metadata"));
+			ExecuteProcessFinishedReporter(InTakeId, MoveTemp(Result));
 		}
 	}
 	else
 	{
-		FCaptureDataConverterError Error = Result.StealError();
-
-		FText Message = FText::Join(FText::FromString(TEXT("\n\r")), Error.GetErrors());
-		ExecuteProcessFinishedReporter(InTakeId,
-									   MakeError(FIngestCapability_Error::ConversionError, Message.ToString()));
-	}
-}
-
-FCaptureDataConverterResult<void> UGenericLiveLinkCPSDevice::RunConversion(int32 InTakeId,
-																		   TStrongObjectPtr<const UIngestCapability_Options> InIngestOptions,
-																		   TSharedPtr<UE::CaptureManager::FTaskProgress> InTaskProgress,
-																		   FCaptureDataConverterParams& InParams)
-{
-	using namespace UE::CaptureManager;
-
-	FTaskProgress::FTask ConvertTask = InTaskProgress->StartTask();
-
-	FCaptureConvertVideoOutputParams VideoParams;
-	VideoParams.Format = InIngestOptions->Video.Format;
-	VideoParams.ImageFileName = InIngestOptions->Video.FileNamePrefix;
-	VideoParams.OutputPixelFormat = Private::ConvertPixelFormat(InIngestOptions->Video.PixelFormat);
-
-	InParams.VideoOutputParams = MoveTemp(VideoParams);
-
-	FCaptureConvertAudioOutputParams AudioParams;
-	AudioParams.Format = InIngestOptions->Audio.Format;
-	AudioParams.AudioFileName = InIngestOptions->Audio.FileNamePrefix;
-
-	InParams.AudioOutputParams = MoveTemp(AudioParams);
-
-	FCaptureConvertDepthOutputParams DepthParams;
-	DepthParams.ImageFileName = TEXT("depth");
-
-	InParams.DepthOutputParams = MoveTemp(DepthParams);
-
-	FCaptureConvertCalibrationOutputParams CalibrationParams;
-	CalibrationParams.FileName = TEXT("calibration");
-
-	InParams.CalibrationOutputParams = MoveTemp(CalibrationParams);
-
-	FCaptureDataConverter::FProgressReporter ProgressReporter =
-		FCaptureDataConverter::FProgressReporter::CreateLambda([ConvertTask](const double InProgress) mutable
-															   {
-																   ConvertTask.Update(InProgress);
-															   });
-
-	TUniquePtr<FCaptureDataConverter>& Converter = Converters.FindOrAdd(InTakeId, MakeUnique<FCaptureDataConverter>());
-	return Converter->Run(InParams, MoveTemp(ProgressReporter));
-}
-
-void UGenericLiveLinkCPSDevice::UploadTakeStateHandler(const FGuid& InTakeUploadId, double InProgress)
-{
-	if (!ensure(InTakeUploadId.IsValid()))
-	{
-		return;
-	}
-
-	FScopeLock Lock(&Mutex);
-	if (FContext* Context = ImportContexts.Find(InTakeUploadId))
-	{
-		Context->Task.Update(InProgress);
-	}
-}
-
-void UGenericLiveLinkCPSDevice::UploadTakeFinished(const FGuid& InTakeUploadId, const FString& InMessage, int32 InCode)
-{
-	using namespace UE::CaptureManager;
+		TValueOrError<void, FIngestCapability_Error> Result = MakeError(FIngestCapability_Error::DownloaderError, InResult.GetError().GetMessage());
 
-	if (!ensure(InTakeUploadId.IsValid()))
-	{
-		return;
+		ExecuteProcessFinishedReporter(InTakeId, MoveTemp(Result));
 	}
 
-	FScopeLock Lock(&Mutex);
-	if (const FContext* Context = ImportContexts.Find(InTakeUploadId))
-	{
-		const bool bShouldCleanConversionDirectory = GetDefault<UCaptureManagerSettings>()->bShouldCleanConversionDirectory;
-		if (bShouldCleanConversionDirectory)
-		{
-			const FString ConversionDirectory = FPaths::Combine(Context->IngestOptions->WorkingDirectory,
-																Context->TakeName + TEXT("_") + InTakeUploadId.ToString());
-			IFileManager::Get().DeleteDirectory(*ConversionDirectory, false, true);
-		}
-
-		if (InMessage.IsEmpty())
-		{
-			ExecuteProcessFinishedReporter(Context->TakeId, MakeValue());
-		}
-		else
-		{
-			ExecuteProcessFinishedReporter(Context->TakeId, MakeError(FIngestCapability_Error::UploaderError, InMessage));
-		}
-
-		ImportContexts.Remove(InTakeUploadId);
-
-		WorkerForTake.Remove(Context->TakeId);
-	}
+	IFileManager::Get().DeleteDirectory(*DownloadedTake, false, true);
 }
 
 TOptional<FTakeMetadata> UGenericLiveLinkCPSDevice::ParseTake(const FString& InTakeName)
diff --git a/Engine/Plugins/VirtualProduction/CaptureManager/CaptureManagerApp/Source/GenericLiveLinkDevice/Private/CPS/GenericLiveLinkCPSDevice.h b/Engine/Plugins/VirtualProduction/CaptureManager/CaptureManagerApp/Source/GenericLiveLinkDevice/Private/CPS/GenericLiveLinkCPSDevice.h
--- a/Engine/Plugins/VirtualProduction/CaptureManager/CaptureManagerApp/Source/GenericLiveLinkDevice/Private/CPS/GenericLiveLinkCPSDevice.h
+++ b/Engine/Plugins/VirtualProduction/CaptureManager/CaptureManagerApp/Source/GenericLiveLinkDevice/Private/CPS/GenericLiveLinkCPSDevice.h
@@ -10,16 +10,13 @@
 #include "Misc/Optional.h"
 
 #include "Async/TaskProgress.h"
-#include "LiveLinkHubExportManager.h"
-
-#include "CaptureDataConverter.h"
-
-#include "CaptureCommunicationManager.h"
 
 #include "Protocol/CPSDevice.h"
 #include "Customizations/ToggleConnectActionCustomization.h"
 #include "Customizations/DeviceIpAddressCustomization.h"
 
+#include "BaseIngestLiveLinkDevice.h"
+
 #include "GenericLiveLinkCPSDevice.generated.h"
 
 
@@ -44,8 +41,7 @@ public:
 
 
 UCLASS(BlueprintType)
-class UGenericLiveLinkCPSDevice : public ULiveLinkDevice
-	, public ILiveLinkDeviceCapability_Ingest
+class UGenericLiveLinkCPSDevice final : public UBaseIngestLiveLinkDevice
 	, public ILiveLinkDeviceCapability_Connection
 	, public ILiveLinkDeviceCapability_Recording
 {
@@ -63,14 +59,14 @@ public:
 
 	virtual void OnDeviceAdded(const FGuid InDeviceGuid, ULiveLinkDeviceSettings* InSettings) override;
 	virtual void OnDeviceRemoved() override;
-	virtual void OnSettingChanged(const FPropertyChangedEvent& InPropertyChangedEvent) override;
 	//~ End ULiveLinkDevice interface
 
 private:
+	virtual FString GetFullTakePath(UE::CaptureManager::FTakeId InTakeId) const override;
+
 	//~ Begin ULiveLinkDeviceCapability_Ingest interface
-	virtual void UpdateTakeListAsync_Implementation() override;
-	virtual void RunIngestTakeAsync(int32 InTakeId,
-									const UIngestCapability_Options* InIngestOptions) override;
+	virtual void UpdateTakeList_Implementation() override;
+	virtual void RunIngestTake(int32 InTakeId, const UIngestCapability_Options* InIngestOptions) override;
 	virtual void CancelIngest(int32 InTakeId) override;
 	//~ End ULiveLinkDeviceCapability_Ingest interface
 
@@ -100,47 +96,16 @@ private:
 						  TSharedPtr<UE::CaptureManager::FTaskProgress> InTaskProgress,
 						  TStrongObjectPtr<const UIngestCapability_Options> InIngestOptions);
 
-	void ConvertData(const FString& InTakeName,
-					 UE::CaptureManager::FTakeId InTakeId,
-					 TSharedPtr<UE::CaptureManager::FTaskProgress> InTaskProgress,
-					 TStrongObjectPtr<const UIngestCapability_Options> InIngestOptions);
-
-	FCaptureDataConverterResult<void> RunConversion(int32 InTakeId,
-													TStrongObjectPtr<const UIngestCapability_Options> InIngestOptions,
-													TSharedPtr<UE::CaptureManager::FTaskProgress> InTaskProgress,
-													FCaptureDataConverterParams& InParams);
-
-	void UploadTakeStateHandler(const FGuid& InTakeUploadId, double InProgress);
-	void UploadTakeFinished(const FGuid& InTakeUploadId, const FString& InMessage, int32 InCode);
-
 	TOptional<FTakeMetadata> ParseTake(const FString& InTakeName);
 	FTakeMetadata ParseTakeMetadata(const UE::CaptureManager::FGetTakeMetadataResponse::FTakeObject& InTake);
 
 	void FetchThumbnail(const FString& InTakeName, UE::CaptureManager::FTakeId InTakeId);
 	void FetchThumbnails(TMap<FString, UE::CaptureManager::FTakeId> InNameToIdMap);
 
-	struct FContext
-	{
-		FString TakeName;
-		int32 TakeId;
-		TStrongObjectPtr<const UIngestCapability_Options> IngestOptions;
-		TSharedPtr<UE::CaptureManager::FTaskProgress> TaskProgress;
-		UE::CaptureManager::FTaskProgress::FTask Task;
-	};
-
 	TSharedPtr<UE::CaptureManager::FCPSDevice> Device;
 
-	TSharedPtr<FCommunicationManager> CommunicationManager;
-	TWeakPtr<FUploader> Uploader;
-
-	TMap<FGuid, FContext> ImportContexts;
-
-	FCriticalSection Mutex;
-	TMap<int32, int32> WorkerForTake;
-
-	TMap<int32, TUniquePtr<FCaptureDataConverter>> Converters;
-
-	FGuid DeviceGuid;
+	mutable FCriticalSection DownloadedTakesMutex;
+	TMap<UE::CaptureManager::FTakeId, FString> DownloadedTakes;
 
 	std::atomic_bool bIsRecording = false;
 	std::atomic_bool bIsConnecting = false;
diff --git a/Engine/Plugins/VirtualProduction/CaptureManager/CaptureManagerApp/Source/GenericLiveLinkDevice/Private/Utils/GenericLiveLinkDeviceUtility.cpp b/Engine/Plugins/VirtualProduction/CaptureManager/CaptureManagerApp/Source/GenericLiveLinkDevice/Private/Utils/GenericLiveLinkDeviceUtility.cpp
--- a/Engine/Plugins/VirtualProduction/CaptureManager/CaptureManagerApp/Source/GenericLiveLinkDevice/Private/Utils/GenericLiveLinkDeviceUtility.cpp
+++ b/Engine/Plugins/VirtualProduction/CaptureManager/CaptureManagerApp/Source/GenericLiveLinkDevice/Private/Utils/GenericLiveLinkDeviceUtility.cpp
@@ -5,44 +5,6 @@
 namespace UE::CaptureManager::Private
 {
 
-TSharedRef<FCommunicationManager> GetCommunicationManager()
-{
-	FCaptureCommunicationManagerModule& CaptureCommunicationManagerModule =
-		FModuleManager::LoadModuleChecked<FCaptureCommunicationManagerModule>("CaptureCommunicationManager");
-	return CaptureCommunicationManagerModule.GetManager();
-}
-
-EMediaTexturePixelFormat ConvertPixelFormat(EIngestCapability_ImagePixelFormat InImagePixelFormat)
-{
-	switch (InImagePixelFormat)
-	{
-		case EIngestCapability_ImagePixelFormat::U8_RGB:
-			return EMediaTexturePixelFormat::U8_RGB;
-		case EIngestCapability_ImagePixelFormat::U8_BGR:
-			return EMediaTexturePixelFormat::U8_BGR;
-		case EIngestCapability_ImagePixelFormat::U8_RGBA:
-			return EMediaTexturePixelFormat::U8_RGBA;
-		case EIngestCapability_ImagePixelFormat::U8_BGRA:
-			return EMediaTexturePixelFormat::U8_BGRA;
-		case EIngestCapability_ImagePixelFormat::U8_I444:
-			return EMediaTexturePixelFormat::U8_I444;
-		case EIngestCapability_ImagePixelFormat::U8_I420:
-			return EMediaTexturePixelFormat::U8_I420;
-		case EIngestCapability_ImagePixelFormat::U8_NV12:
-			return EMediaTexturePixelFormat::U8_NV12;
-		case EIngestCapability_ImagePixelFormat::U8_Mono:
-			return EMediaTexturePixelFormat::U8_Mono;
-		case EIngestCapability_ImagePixelFormat::U16_Mono:
-			return EMediaTexturePixelFormat::U16_Mono;
-		case EIngestCapability_ImagePixelFormat::F_Mono:
-			return EMediaTexturePixelFormat::F_Mono;
-		case EIngestCapability_ImagePixelFormat::Undefined:
-		default:
-			check(false);
-			return EMediaTexturePixelFormat::Undefined;
-	}
-}
-
 FString ErrorOriginToString(FTakeMetadataParserError::EOrigin InOrigin)
 {
 	switch (InOrigin)
diff --git a/Engine/Plugins/VirtualProduction/CaptureManager/CaptureManagerApp/Source/GenericLiveLinkDevice/Private/Utils/GenericLiveLinkDeviceUtility.h b/Engine/Plugins/VirtualProduction/CaptureManager/CaptureManagerApp/Source/GenericLiveLinkDevice/Private/Utils/GenericLiveLinkDeviceUtility.h
--- a/Engine/Plugins/VirtualProduction/CaptureManager/CaptureManagerApp/Source/GenericLiveLinkDevice/Private/Utils/GenericLiveLinkDeviceUtility.h
+++ b/Engine/Plugins/VirtualProduction/CaptureManager/CaptureManagerApp/Source/GenericLiveLinkDevice/Private/Utils/GenericLiveLinkDeviceUtility.h
@@ -2,16 +2,11 @@
 
 #pragma once
 
-#include "CaptureCommunicationManagerModule.h"
 #include "CaptureManagerTakeMetadata.h"
-#include "Ingest/IngestCapability_Options.h"
-#include "MediaSample.h"
 
 namespace UE::CaptureManager::Private
 {
 
-TSharedRef<FCommunicationManager> GetCommunicationManager();
-EMediaTexturePixelFormat ConvertPixelFormat(EIngestCapability_ImagePixelFormat InImagePixelFormat);
 FString ErrorOriginToString(FTakeMetadataParserError::EOrigin InOrigin);
 
 }
\ No newline at end of file
diff --git a/Engine/Plugins/VirtualProduction/CaptureManager/CaptureManagerApp/Source/IngestLiveLinkDevice/IngestLiveLinkDevice.Build.cs b/Engine/Plugins/VirtualProduction/CaptureManager/CaptureManagerApp/Source/IngestLiveLinkDevice/IngestLiveLinkDevice.Build.cs
new file mode 100644
--- /dev/null
+++ b/Engine/Plugins/VirtualProduction/CaptureManager/CaptureManagerApp/Source/IngestLiveLinkDevice/IngestLiveLinkDevice.Build.cs
@@ -0,0 +1,43 @@
+// Copyright Epic Games, Inc. All Rights Reserved.
+
+using UnrealBuildTool;
+using System.IO;
+
+public class IngestLiveLinkDevice : ModuleRules
+{
+	public IngestLiveLinkDevice(ReadOnlyTargetRules Target) : base(Target)
+	{
+		PCHUsage = PCHUsageMode.UseExplicitOrSharedPCHs;
+		ShortName = "IngestLiveLinkDevice";
+
+		PublicDependencyModuleNames.AddRange(new string[]
+		{
+			"CaptureUtils",
+			"Core"
+		});
+
+		PrivateDependencyModuleNames.AddRange(new string[]
+		{
+			"CaptureProtocolStack",
+			"CaptureDataConverter",
+			"CaptureCommunicationManager",
+			"Engine",
+			"Json",
+			"LiveLinkDevice",
+			"LiveLinkCapabilities",
+			"Media",
+			"LiveLinkHubCaptureMessaging",
+			"LiveLinkHubExportClient",
+			"LiveLinkHub",
+			"ToolWidgets",
+			"Slate",
+			"SlateCore",
+			"CaptureManagerTakeMetadata",
+			"CaptureManagerSettings",
+			"CoreUObject",
+			"UnrealEd"
+		});
+
+		CppCompileWarningSettings.UnsafeTypeCastWarningLevel = WarningLevel.Error;
+	}
+}
diff --git a/Engine/Plugins/VirtualProduction/CaptureManager/CaptureManagerApp/Source/IngestLiveLinkDevice/Private/BaseIngestLiveLinkDevice.cpp b/Engine/Plugins/VirtualProduction/CaptureManager/CaptureManagerApp/Source/IngestLiveLinkDevice/Private/BaseIngestLiveLinkDevice.cpp
new file mode 100644
--- /dev/null
+++ b/Engine/Plugins/VirtualProduction/CaptureManager/CaptureManagerApp/Source/IngestLiveLinkDevice/Private/BaseIngestLiveLinkDevice.cpp
@@ -0,0 +1,322 @@
+// Copyright Epic Games, Inc. All Rights Reserved.
+
+#include "BaseIngestLiveLinkDevice.h"
+
+#include "HAL/FileManager.h"
+#include "CaptureCommunicationManagerModule.h"
+#include "Settings/CaptureManagerSettings.h"
+
+#include "Utils/IngestLiveLinkDeviceUtils.h"
+#include "Utils/IngestLiveLinkDeviceLog.h"
+
+#include "CaptureDataConverter.h"
+
+#include "CaptureCommunicationManager.h"
+
+DEFINE_LOG_CATEGORY(LogIngestLiveLinkDevice);
+
+class UBaseIngestLiveLinkDevice::FImpl
+{
+public:
+
+	FImpl(UBaseIngestLiveLinkDevice* InOwner, TSharedPtr<FCommunicationManager> InCommunicationManager);
+
+	void IngestTake(int32 InTakeId, const UIngestCapability_Options* InIngestOptions, TSharedPtr<UE::CaptureManager::FTaskProgress> InTaskProgress);
+	void CancelIngest(int32 InTakeId);
+
+private:
+
+	FCaptureDataConverterResult<void> RunConversion(int32 InTakeId,
+													TStrongObjectPtr<const UIngestCapability_Options> InIngestOptions,
+													TSharedPtr<UE::CaptureManager::FTaskProgress> InTaskProgress,
+													FCaptureDataConverterParams& InParams);
+
+	void UploadTakeStateHandler(const FGuid& InTakeUploadId, double InProgress);
+	void UploadTakeFinished(const FGuid& InTakeUploadId, const FString& InMessage, int32 InCode);
+
+	struct FContext
+	{
+		FString TakeName;
+		int32 TakeId = INDEX_NONE;
+		TStrongObjectPtr<const UIngestCapability_Options> IngestOptions;
+		TSharedPtr<UE::CaptureManager::FTaskProgress> TaskProgress;
+		UE::CaptureManager::FTaskProgress::FTask Task;
+	};
+
+	UBaseIngestLiveLinkDevice* Owner;
+
+	TSharedPtr<FCommunicationManager> CommunicationManager;
+	TWeakPtr<FUploader> Uploader;
+
+	TMap<FGuid, FContext> ImportContexts;
+
+	FCriticalSection Mutex;
+	TMap<int32, int32> WorkerForTake;
+
+	TMap<int32, TUniquePtr<FCaptureDataConverter>> Converters;
+};
+
+UBaseIngestLiveLinkDevice::FImpl::FImpl(UBaseIngestLiveLinkDevice* InOwner, TSharedPtr<FCommunicationManager> InCommunicationManager)
+	: Owner(InOwner)
+	, CommunicationManager(MoveTemp(InCommunicationManager))
+{
+}
+
+void UBaseIngestLiveLinkDevice::FImpl::IngestTake(int32 InTakeId, const UIngestCapability_Options* InIngestOptions, TSharedPtr<UE::CaptureManager::FTaskProgress> InTaskProgress)
+{
+	TStrongObjectPtr<const UIngestCapability_Options> IngestOptions(InIngestOptions);
+
+	using namespace UE::CaptureManager;
+
+	FString TakeFullPath = Owner->GetFullTakePath(InTakeId);
+
+	if (TakeFullPath.IsEmpty())
+	{
+		FString Message = TEXT("Failed to find the take directory for the specified take");
+		Owner->ExecuteProcessFinishedReporter(InTakeId, MakeError(FIngestCapability_Error::InvalidArgument, MoveTemp(Message)));
+		return;
+	}
+
+	TOptional<FTakeMetadata> MaybeTakeMetadata = Owner->GetTakeMetadata(InTakeId);
+	if (!MaybeTakeMetadata.IsSet())
+	{
+		FString Message = TEXT("Failed to find the take metadata for the specified take");
+		Owner->ExecuteProcessFinishedReporter(InTakeId, MakeError(FIngestCapability_Error::InvalidArgument, MoveTemp(Message)));
+		return;
+	}
+
+	FTakeMetadata TakeMetadata = MoveTemp(MaybeTakeMetadata.GetValue());
+
+	FString TakeName = TakeMetadata.Slate + TEXT("_") + FString::FromInt(TakeMetadata.TakeNumber);
+
+	const FString TakeConversionDirectory = IngestOptions->WorkingDirectory;
+
+	FGuid TakeUploadId = FGuid::NewGuid();
+
+	FCaptureDataConverterParams Params;
+	Params.TakeName = TakeName;
+	Params.TakeMetadata = TakeMetadata;
+	Params.TakeOriginDirectory = MoveTemp(TakeFullPath);
+	Params.TakeOutputDirectory = FPaths::Combine(TakeConversionDirectory, TakeName + TEXT("_") + TakeUploadId.ToString());
+
+	FCaptureDataConverterResult<void> Result = RunConversion(InTakeId, IngestOptions, InTaskProgress, Params);
+
+	Converters.Remove(InTakeId);
+
+	if (Result.HasValue())
+	{
+		if (!Uploader.IsValid())
+		{
+			Uploader = CommunicationManager->GetUploader(Owner->DeviceGuid,
+														 Owner->GetName(),
+														 FUploader::FUploadStateUpdate::CreateRaw(this, &FImpl::UploadTakeStateHandler),
+														 FUploader::FUploadFinished::CreateRaw(this, &FImpl::UploadTakeFinished));
+		}
+
+		if (TSharedPtr<FUploader> SharedUploader = Uploader.Pin())
+		{
+			FScopeLock Lock(&Mutex);
+
+			FTaskProgress::FTask UploadTask = InTaskProgress->StartTask();
+
+			ImportContexts.Add(TakeUploadId, { TakeName, InTakeId, IngestOptions, InTaskProgress, UploadTask });
+
+			int32 Worker = SharedUploader->UploadTake(TakeUploadId, Params.TakeOutputDirectory, MoveTemp(TakeMetadata));
+			WorkerForTake.Add(InTakeId, Worker);
+		}
+		else
+		{
+			IFileManager::Get().DeleteDirectory(*Params.TakeOutputDirectory, false, true);
+
+			FString Message = TEXT("Failed to upload take: Uploader not found");
+			Owner->ExecuteProcessFinishedReporter(InTakeId, MakeError(FIngestCapability_Error::UploaderNotFound, MoveTemp(Message)));
+		}
+	}
+	else
+	{
+		IFileManager::Get().DeleteDirectory(*Params.TakeOutputDirectory, false, true);
+
+		FCaptureDataConverterError Error = Result.StealError();
+
+		FText Message = FText::Join(FText::FromString(TEXT("\n\r")), Error.GetErrors());
+		Owner->ExecuteProcessFinishedReporter(InTakeId, MakeError(FIngestCapability_Error::ConversionError, Message.ToString()));
+	}
+}
+
+void UBaseIngestLiveLinkDevice::FImpl::CancelIngest(int32 InTakeId)
+{
+	if (TUniquePtr<FCaptureDataConverter>* Converter = Converters.Find(InTakeId))
+	{
+		(*Converter)->Cancel();
+	}
+
+	int32 WorkerId = INDEX_NONE;
+
+	{
+		FScopeLock Lock(&Mutex);
+		if (int32* WorkerIdElem = WorkerForTake.Find(InTakeId))
+		{
+			WorkerId = *WorkerIdElem;
+		}
+	}
+
+	if (WorkerId != INDEX_NONE)
+	{
+		if (TSharedPtr<FUploader> SharedUploader = Uploader.Pin())
+		{
+			SharedUploader->AbortUpload(WorkerId);
+		}
+		else
+		{
+			// If `UploadTakeFinished` is not called (eg. connection between LLH and UE broke), `Uploader` will be deleted, but UI will stay in state where import is in progress 99.9%.
+			// Clicking stop button at that moment will not enter if block above, as uploader does not exist, so in that case we have to manually cleanup the workspace and put UI in proper state.
+			FScopeLock Lock(&Mutex);
+
+			TMap<FGuid, FContext> FilteredMap = ImportContexts.FilterByPredicate([InTakeId](const TPair<FGuid, FContext>& Item) -> bool
+			{
+				return Item.Value.TakeId == InTakeId;
+			});
+
+			check(FilteredMap.Num() == 1);
+			TMap<FGuid, FContext>::TConstIterator FoundContext = FilteredMap.CreateConstIterator();
+			check(FoundContext);
+
+			const FContext& Context = FoundContext->Value;
+
+			const FString TakeDirectory = FPaths::Combine(Context.IngestOptions->WorkingDirectory,
+														  Context.TakeName + TEXT("_") + FoundContext->Key.ToString());
+			IFileManager::Get().DeleteDirectory(*TakeDirectory, false, true);
+
+			FString Message = TEXT("Aborted by user");
+			Owner->ExecuteProcessFinishedReporter(InTakeId, MakeError(FIngestCapability_Error::AbortedByUser, MoveTemp(Message)));
+
+			ImportContexts.Remove(FoundContext->Key);
+
+			WorkerForTake.Remove(InTakeId);
+		}
+	}
+}
+
+
+FCaptureDataConverterResult<void> UBaseIngestLiveLinkDevice::FImpl::RunConversion(int32 InTakeId,
+																				  TStrongObjectPtr<const UIngestCapability_Options> InIngestOptions,
+																				  TSharedPtr<UE::CaptureManager::FTaskProgress> InTaskProgress,
+																				  FCaptureDataConverterParams& InParams)
+{
+	using namespace UE::CaptureManager;
+
+	FTaskProgress::FTask ConvertTask = InTaskProgress->StartTask();
+
+	FCaptureConvertVideoOutputParams VideoParams;
+	VideoParams.Format = InIngestOptions->Video.Format;
+	VideoParams.ImageFileName = InIngestOptions->Video.FileNamePrefix;
+	VideoParams.OutputPixelFormat = Private::ConvertPixelFormat(InIngestOptions->Video.PixelFormat);
+
+	InParams.VideoOutputParams = MoveTemp(VideoParams);
+
+	FCaptureConvertAudioOutputParams AudioParams;
+	AudioParams.Format = InIngestOptions->Audio.Format;
+	AudioParams.AudioFileName = InIngestOptions->Audio.FileNamePrefix;
+
+	InParams.AudioOutputParams = MoveTemp(AudioParams);
+
+	FCaptureConvertDepthOutputParams DepthParams;
+	DepthParams.ImageFileName = TEXT("depth");
+
+	InParams.DepthOutputParams = MoveTemp(DepthParams);
+
+	FCaptureConvertCalibrationOutputParams CalibrationParams;
+	CalibrationParams.FileName = TEXT("calibration");
+
+	InParams.CalibrationOutputParams = MoveTemp(CalibrationParams);
+
+	FCaptureDataConverter::FProgressReporter ProgressReporter =
+		FCaptureDataConverter::FProgressReporter::CreateLambda([ConvertTask](const double InProgress) mutable
+	{
+		ConvertTask.Update(InProgress);
+	});
+
+	TUniquePtr<FCaptureDataConverter>& Converter = Converters.FindOrAdd(InTakeId, MakeUnique<FCaptureDataConverter>());
+	return Converter->Run(InParams, MoveTemp(ProgressReporter));
+}
+
+void UBaseIngestLiveLinkDevice::FImpl::UploadTakeStateHandler(const FGuid& InTakeUploadId, double InProgress)
+{
+	if (!ensure(InTakeUploadId.IsValid()))
+	{
+		return;
+	}
+
+	FScopeLock Lock(&Mutex);
+	if (FContext* Context = ImportContexts.Find(InTakeUploadId))
+	{
+		Context->Task.Update(InProgress);
+	}
+}
+
+void UBaseIngestLiveLinkDevice::FImpl::UploadTakeFinished(const FGuid& InTakeUploadId, const FString& InMessage, int32 InCode)
+{
+	using namespace UE::CaptureManager;
+
+	if (!ensure(InTakeUploadId.IsValid()))
+	{
+		return;
+	}
+
+	FScopeLock Lock(&Mutex);
+	if (const FContext* Context = ImportContexts.Find(InTakeUploadId))
+	{
+		const bool bShouldCleanConversionDirectory = GetDefault<UCaptureManagerSettings>()->bShouldCleanConversionDirectory;
+		if (bShouldCleanConversionDirectory)
+		{
+			const FString ConversionDirectory = FPaths::Combine(Context->IngestOptions->WorkingDirectory,
+																Context->TakeName + TEXT("_") + InTakeUploadId.ToString());
+			IFileManager::Get().DeleteDirectory(*ConversionDirectory, false, true);
+		}
+
+		if (InMessage.IsEmpty())
+		{
+			Owner->ExecuteProcessFinishedReporter(Context->TakeId, MakeValue());
+		}
+		else
+		{
+			Owner->ExecuteProcessFinishedReporter(Context->TakeId, MakeError(FIngestCapability_Error::UploaderError, InMessage));
+		}
+
+		ImportContexts.Remove(InTakeUploadId);
+
+		WorkerForTake.Remove(Context->TakeId);
+	}
+}
+
+
+FGuid UBaseIngestLiveLinkDevice::GetDeviceGuid() const
+{
+	return DeviceGuid;
+}
+
+void UBaseIngestLiveLinkDevice::OnDeviceAdded(const FGuid InDeviceGuid, ULiveLinkDeviceSettings* InSettings)
+{
+	DeviceGuid = InDeviceGuid;
+
+	Impl = MakePimpl<FImpl>(this, UE::CaptureManager::Private::GetCommunicationManager());
+
+	Super::OnDeviceAdded(DeviceGuid, InSettings);
+}
+
+void UBaseIngestLiveLinkDevice::OnDeviceRemoved()
+{
+	RemoveAllTakes();
+
+	Super::OnDeviceRemoved();
+}
+
+void UBaseIngestLiveLinkDevice::IngestTake(int32 InTakeId, const UIngestCapability_Options* InIngestOptions, TSharedPtr<UE::CaptureManager::FTaskProgress> InTaskProgress)
+{
+	Impl->IngestTake(InTakeId, InIngestOptions, MoveTemp(InTaskProgress));
+}
+
+void UBaseIngestLiveLinkDevice::CancelIngest(int32 InTakeId)
+{
+	Impl->CancelIngest(InTakeId);
+}
\ No newline at end of file
diff --git a/Engine/Plugins/VirtualProduction/CaptureManager/CaptureManagerApp/Source/IngestLiveLinkDevice/Private/IngestLiveLinkDeviceModule.cpp b/Engine/Plugins/VirtualProduction/CaptureManager/CaptureManagerApp/Source/IngestLiveLinkDevice/Private/IngestLiveLinkDeviceModule.cpp
new file mode 100644
--- /dev/null
+++ b/Engine/Plugins/VirtualProduction/CaptureManager/CaptureManagerApp/Source/IngestLiveLinkDevice/Private/IngestLiveLinkDeviceModule.cpp
@@ -0,0 +1,5 @@
+// Copyright Epic Games, Inc. All Rights Reserved.
+
+#include "Modules/ModuleManager.h"
+
+IMPLEMENT_MODULE(FDefaultModuleImpl, IngestLiveLinkDevice)
\ No newline at end of file
diff --git a/Engine/Plugins/VirtualProduction/CaptureManager/CaptureManagerApp/Source/IngestLiveLinkDevice/Private/Utils/IngestLiveLinkDeviceLog.h b/Engine/Plugins/VirtualProduction/CaptureManager/CaptureManagerApp/Source/IngestLiveLinkDevice/Private/Utils/IngestLiveLinkDeviceLog.h
new file mode 100644
--- /dev/null
+++ b/Engine/Plugins/VirtualProduction/CaptureManager/CaptureManagerApp/Source/IngestLiveLinkDevice/Private/Utils/IngestLiveLinkDeviceLog.h
@@ -0,0 +1,7 @@
+// Copyright Epic Games, Inc. All Rights Reserved.
+
+#pragma once
+
+#include "Logging/LogMacros.h"
+
+DECLARE_LOG_CATEGORY_EXTERN(LogIngestLiveLinkDevice, Log, All);
diff --git a/Engine/Plugins/VirtualProduction/CaptureManager/CaptureManagerApp/Source/IngestLiveLinkDevice/Private/Utils/IngestLiveLinkDeviceUtils.cpp b/Engine/Plugins/VirtualProduction/CaptureManager/CaptureManagerApp/Source/IngestLiveLinkDevice/Private/Utils/IngestLiveLinkDeviceUtils.cpp
new file mode 100644
--- /dev/null
+++ b/Engine/Plugins/VirtualProduction/CaptureManager/CaptureManagerApp/Source/IngestLiveLinkDevice/Private/Utils/IngestLiveLinkDeviceUtils.cpp
@@ -0,0 +1,60 @@
+// Copyright Epic Games, Inc. All Rights Reserved.
+
+#include "IngestLiveLinkDeviceUtils.h"
+
+namespace UE::CaptureManager::Private
+{
+
+TSharedRef<FCommunicationManager> GetCommunicationManager()
+{
+	FCaptureCommunicationManagerModule& CaptureCommunicationManagerModule =
+		FModuleManager::LoadModuleChecked<FCaptureCommunicationManagerModule>("CaptureCommunicationManager");
+	return CaptureCommunicationManagerModule.GetManager();
+}
+
+EMediaTexturePixelFormat ConvertPixelFormat(EIngestCapability_ImagePixelFormat InImagePixelFormat)
+{
+	switch (InImagePixelFormat)
+	{
+		case EIngestCapability_ImagePixelFormat::U8_RGB:
+			return EMediaTexturePixelFormat::U8_RGB;
+		case EIngestCapability_ImagePixelFormat::U8_BGR:
+			return EMediaTexturePixelFormat::U8_BGR;
+		case EIngestCapability_ImagePixelFormat::U8_RGBA:
+			return EMediaTexturePixelFormat::U8_RGBA;
+		case EIngestCapability_ImagePixelFormat::U8_BGRA:
+			return EMediaTexturePixelFormat::U8_BGRA;
+		case EIngestCapability_ImagePixelFormat::U8_I444:
+			return EMediaTexturePixelFormat::U8_I444;
+		case EIngestCapability_ImagePixelFormat::U8_I420:
+			return EMediaTexturePixelFormat::U8_I420;
+		case EIngestCapability_ImagePixelFormat::U8_NV12:
+			return EMediaTexturePixelFormat::U8_NV12;
+		case EIngestCapability_ImagePixelFormat::U8_Mono:
+			return EMediaTexturePixelFormat::U8_Mono;
+		case EIngestCapability_ImagePixelFormat::U16_Mono:
+			return EMediaTexturePixelFormat::U16_Mono;
+		case EIngestCapability_ImagePixelFormat::F_Mono:
+			return EMediaTexturePixelFormat::F_Mono;
+		case EIngestCapability_ImagePixelFormat::Undefined:
+		default:
+			check(false);
+			return EMediaTexturePixelFormat::Undefined;
+	}
+}
+
+FString ErrorOriginToString(FTakeMetadataParserError::EOrigin InOrigin)
+{
+	switch (InOrigin)
+	{
+		case FTakeMetadataParserError::Reader:
+			return TEXT("Reader");
+		case FTakeMetadataParserError::Validator:
+			return TEXT("Validator");
+		case FTakeMetadataParserError::Parser:
+		default:
+			return TEXT("Parser");
+	}
+}
+
+}
\ No newline at end of file
diff --git a/Engine/Plugins/VirtualProduction/CaptureManager/CaptureManagerApp/Source/IngestLiveLinkDevice/Private/Utils/IngestLiveLinkDeviceUtils.h b/Engine/Plugins/VirtualProduction/CaptureManager/CaptureManagerApp/Source/IngestLiveLinkDevice/Private/Utils/IngestLiveLinkDeviceUtils.h
new file mode 100644
--- /dev/null
+++ b/Engine/Plugins/VirtualProduction/CaptureManager/CaptureManagerApp/Source/IngestLiveLinkDevice/Private/Utils/IngestLiveLinkDeviceUtils.h
@@ -0,0 +1,17 @@
+// Copyright Epic Games, Inc. All Rights Reserved.
+
+#pragma once
+
+#include "CaptureCommunicationManagerModule.h"
+#include "CaptureManagerTakeMetadata.h"
+#include "Ingest/IngestCapability_Options.h"
+#include "MediaSample.h"
+
+namespace UE::CaptureManager::Private
+{
+
+TSharedRef<FCommunicationManager> GetCommunicationManager();
+EMediaTexturePixelFormat ConvertPixelFormat(EIngestCapability_ImagePixelFormat InImagePixelFormat);
+FString ErrorOriginToString(FTakeMetadataParserError::EOrigin InOrigin);
+
+}
\ No newline at end of file
diff --git a/Engine/Plugins/VirtualProduction/CaptureManager/CaptureManagerApp/Source/IngestLiveLinkDevice/Public/BaseIngestLiveLinkDevice.h b/Engine/Plugins/VirtualProduction/CaptureManager/CaptureManagerApp/Source/IngestLiveLinkDevice/Public/BaseIngestLiveLinkDevice.h
new file mode 100644
--- /dev/null
+++ b/Engine/Plugins/VirtualProduction/CaptureManager/CaptureManagerApp/Source/IngestLiveLinkDevice/Public/BaseIngestLiveLinkDevice.h
@@ -0,0 +1,90 @@
+// Copyright Epic Games, Inc. All Rights Reserved.
+
+#pragma once
+
+#include "LiveLinkDevice.h"
+
+#include "Ingest/LiveLinkDeviceCapability_Ingest.h"
+
+#include "Async/TaskProgress.h"
+#include "Misc/Optional.h"
+
+#include "BaseIngestLiveLinkDevice.generated.h"
+
+/**
+ * Base class that provides default implementations of core ingest capability functions.
+ */
+UCLASS(Abstract)
+class INGESTLIVELINKDEVICE_API UBaseIngestLiveLinkDevice : public ULiveLinkDevice
+	, public ILiveLinkDeviceCapability_Ingest
+{
+	GENERATED_BODY()
+
+public:
+
+	/**
+	 * Function provides the Device Guid defined by the LiveLink Device Subsystem during device
+	 * @returns Device unique ID
+	 */
+	FGuid GetDeviceGuid() const;
+
+	//~ Begin ULiveLinkDevice interface
+
+	/**
+	 * Inherited function from the ULiveLinkDevice that is called when device is added
+	 * @note Users are required to invoke this function if it is inherited in the users code.
+	 * @params DeviceGuid Device unique ID
+	 * @params Settigns Device settings
+	 */
+	virtual void OnDeviceAdded(const FGuid InDeviceGuid, ULiveLinkDeviceSettings* InSettings) override;
+
+	/**
+	 * Inherited function from the ULiveLinkDevice that is called when device is removed
+	 * @note Users are required to invoke this function if it is inherited in the users code.
+	 */
+	virtual void OnDeviceRemoved() override;
+
+	//~ End ULiveLinkDevice interface
+
+protected:
+
+	/**
+	 * Function implements the default behaviour for conversion and upload to the UE of the take data
+	 * @params TakeId Take identifier
+	 * @params IngestOptions Ingest options used for conversion process
+	 * @params TaskProgress Helper class that accumulates progress and reports the progress to the user
+	 */
+	void IngestTake(int32 InTakeId, const UIngestCapability_Options* InIngestOptions, TSharedPtr<UE::CaptureManager::FTaskProgress> InTaskProgress);
+
+	/**
+	 * Pure virtual function that returns the full path to the take data. Full path to the take data is required for the conversion step.
+	 * @params TakeId Take identifier
+	 * @returns Full path to the take data
+	 */
+	virtual FString GetFullTakePath(UE::CaptureManager::FTakeId InTakeId) const PURE_VIRTUAL(UBaseIngestLiveLinkDevice::GetFullTakePath, return FString(););
+
+	//~ Begin ULiveLinkDeviceCapability_Ingest interface
+
+	/**
+	 * Inherited function from the ILiveLinkDeviceCapability_Ingest. The function is required to be implemented.
+	 * @note In it's implementation, user may choose to call IngestTake function or to implement custom conversion and upload step
+	 * @params TakeId Take identifier
+	 * @params IngestOptions Ingest options used for conversion process
+	 */
+	virtual void RunIngestTake(int32 InTakeId, const UIngestCapability_Options* InIngestOptions) override PURE_VIRTUAL(UBaseIngestLiveLinkDevice::RunIngestTake);
+
+	/**
+	 * Inherited function from the ILiveLinkDeviceCapability_Ingest. The function is required to be implemented.
+	 * @note In it's default implementation, it will abort the default ingest process.
+	 * @params TakeId Take identifier
+	 */
+	virtual void CancelIngest(int32 InTakeId) override;
+	//~ End ULiveLinkDeviceCapability_Ingest interface
+
+private:
+
+	class FImpl;
+	TPimplPtr<FImpl> Impl;
+
+	FGuid DeviceGuid;
+};
diff --git a/Engine/Plugins/VirtualProduction/CaptureManager/CaptureManagerApp/Source/LiveLinkCapabilities/Private/Ingest/IngestCapability_ProcessHandle.cpp b/Engine/Plugins/VirtualProduction/CaptureManager/CaptureManagerApp/Source/LiveLinkCapabilities/Private/Ingest/IngestCapability_ProcessHandle.cpp
--- a/Engine/Plugins/VirtualProduction/CaptureManager/CaptureManagerApp/Source/LiveLinkCapabilities/Private/Ingest/IngestCapability_ProcessHandle.cpp
+++ b/Engine/Plugins/VirtualProduction/CaptureManager/CaptureManagerApp/Source/LiveLinkCapabilities/Private/Ingest/IngestCapability_ProcessHandle.cpp
@@ -36,7 +36,7 @@ void UIngestCapability_ProcessHandle::Run(const UIngestCapability_Options* InIng
 {
 	if (ensure(InIngestOptions))
 	{
-		Owner->RunIngestTakeAsync(TakeId, InIngestOptions);
+		Owner->RunIngestTake(TakeId, InIngestOptions);
 	}
 }
 
diff --git a/Engine/Plugins/VirtualProduction/CaptureManager/CaptureManagerApp/Source/LiveLinkCapabilities/Public/Ingest/LiveLinkDeviceCapability_Ingest.h b/Engine/Plugins/VirtualProduction/CaptureManager/CaptureManagerApp/Source/LiveLinkCapabilities/Public/Ingest/LiveLinkDeviceCapability_Ingest.h
--- a/Engine/Plugins/VirtualProduction/CaptureManager/CaptureManagerApp/Source/LiveLinkCapabilities/Public/Ingest/LiveLinkDeviceCapability_Ingest.h
+++ b/Engine/Plugins/VirtualProduction/CaptureManager/CaptureManagerApp/Source/LiveLinkCapabilities/Public/Ingest/LiveLinkDeviceCapability_Ingest.h
@@ -49,7 +49,7 @@ public:
 	UIngestCapability_ProcessHandle* PrepareIngestProcess(int32 InTakeId);
 
 	UFUNCTION(BlueprintCallable, BlueprintNativeEvent, Category = "Live Link Device|Ingest")
-	void UpdateTakeListAsync();
+	void UpdateTakeList();
 
 	UFUNCTION(BlueprintCallable, BlueprintNativeEvent, Category = "Live Link Device|Ingest")
 	UIngestCapability_TakeInformation* GetTakeInformation(int32 InTakeId) const;
@@ -61,7 +61,7 @@ public:
 	UIngestCapability_UpdateTakeListCallback* CreateUpdateTakeListCallback();
 
 	UIngestCapability_ProcessHandle* PrepareIngestProcess_Implementation(int32 InTakeId);
-	virtual void UpdateTakeListAsync_Implementation() PURE_VIRTUAL(ILiveLinkDeviceCapability_Ingest::UpdateTakeListAsync_Implementation);
+	virtual void UpdateTakeList_Implementation() PURE_VIRTUAL(ILiveLinkDeviceCapability_Ingest::UpdateTakeList_Implementation);
 	UIngestCapability_TakeInformation* GetTakeInformation_Implementation(int32 InTakeId) const;
 	TArray<int32> GetTakeIdentifiers_Implementation() const;
 	UIngestCapability_UpdateTakeListCallback* CreateUpdateTakeListCallback_Implementation();
@@ -85,7 +85,7 @@ private:
 		CreateIngestProcessHandle(int32 InTakeId, ILiveLinkDeviceCapability_Ingest* InOwner);
 
 	virtual void CancelIngest(int32 InTakeId) = 0;
-	virtual void RunIngestTakeAsync(int32 InTakeId, const UIngestCapability_Options* InIngestOptions) = 0;
+	virtual void RunIngestTake(int32 InTakeId, const UIngestCapability_Options* InIngestOptions) = 0;
 
 	mutable FCriticalSection TakeMetadataMutex;
 	TMap<int32, FTakeMetadata> TakeMetadataMap;
