Example Live Link Devices Module

diff --git a/Engine/Plugins/VirtualProduction/CaptureManager/CaptureManagerApp/Source/ExampleLiveLinkDevices/ExampleLiveLinkDevices.Build.cs b/Engine/Plugins/VirtualProduction/CaptureManager/CaptureManagerApp/Source/ExampleLiveLinkDevices/ExampleLiveLinkDevices.Build.cs
new file mode 100644
--- /dev/null
+++ b/Engine/Plugins/VirtualProduction/CaptureManager/CaptureManagerApp/Source/ExampleLiveLinkDevices/ExampleLiveLinkDevices.Build.cs
@@ -0,0 +1,29 @@
+// Copyright Epic Games, Inc. All Rights Reserved.
+
+using UnrealBuildTool;
+using System.IO;
+
+public class ExampleLiveLinkDevices : ModuleRules
+{
+	public ExampleLiveLinkDevices(ReadOnlyTargetRules Target) : base(Target)
+	{
+		PCHUsage = PCHUsageMode.UseExplicitOrSharedPCHs;
+
+		PublicDependencyModuleNames.AddRange(new string[]
+		{
+			"CaptureUtils",
+			"Core"
+		});
+
+		PrivateDependencyModuleNames.AddRange(new string[]
+		{
+			"LiveLinkDevice",
+			"IngestLiveLinkDevice",
+			"LiveLinkCapabilities",
+			"CaptureManagerTakeMetadata",
+			"CoreUObject",
+		});
+
+		CppCompileWarningSettings.UnsafeTypeCastWarningLevel = WarningLevel.Error;
+	}
+}
diff --git a/Engine/Plugins/VirtualProduction/CaptureManager/CaptureManagerApp/Source/ExampleLiveLinkDevices/Private/ExampleLiveLinkDevicesModule.cpp b/Engine/Plugins/VirtualProduction/CaptureManager/CaptureManagerApp/Source/ExampleLiveLinkDevices/Private/ExampleLiveLinkDevicesModule.cpp
new file mode 100644
--- /dev/null
+++ b/Engine/Plugins/VirtualProduction/CaptureManager/CaptureManagerApp/Source/ExampleLiveLinkDevices/Private/ExampleLiveLinkDevicesModule.cpp
@@ -0,0 +1,5 @@
+// Copyright Epic Games, Inc. All Rights Reserved.
+
+#include "Modules/ModuleManager.h"
+
+IMPLEMENT_MODULE(FDefaultModuleImpl, ExampleLiveLinkDevices);
\ No newline at end of file
diff --git a/Engine/Plugins/VirtualProduction/CaptureManager/CaptureManagerApp/Source/ExampleLiveLinkDevices/Private/NetworkIngest/ExampleNetworkIngestDevice.cpp b/Engine/Plugins/VirtualProduction/CaptureManager/CaptureManagerApp/Source/ExampleLiveLinkDevices/Private/NetworkIngest/ExampleNetworkIngestDevice.cpp
new file mode 100644
--- /dev/null
+++ b/Engine/Plugins/VirtualProduction/CaptureManager/CaptureManagerApp/Source/ExampleLiveLinkDevices/Private/NetworkIngest/ExampleNetworkIngestDevice.cpp
@@ -0,0 +1,168 @@
+// Copyright Epic Games, Inc. All Rights Reserved.
+
+#include "ExampleNetworkIngestDevice.h"
+
+#include "Async/TaskProgress.h"
+
+const UExampleNetworkIngestDeviceSettings* UExampleNetworkIngestDevice::GetSettings() const
+{
+	return CastChecked<UExampleNetworkIngestDeviceSettings>(Settings);
+}
+
+TSubclassOf<ULiveLinkDeviceSettings> UExampleNetworkIngestDevice::StaticGetSettingsClass() const
+{
+	return UExampleNetworkIngestDeviceSettings::StaticClass();
+}
+
+FText UExampleNetworkIngestDevice::GetDisplayName() const
+{
+	return FText::FromString(GetSettings()->DisplayName);
+}
+
+EDeviceHealth UExampleNetworkIngestDevice::GetDeviceHealth() const
+{
+	return EDeviceHealth::Nominal;
+}
+
+FText UExampleNetworkIngestDevice::GetHealthText() const
+{
+	return FText::FromString("Example Health");
+}
+
+void UExampleNetworkIngestDevice::OnDeviceAdded(const FGuid InDeviceGuid, ULiveLinkDeviceSettings* InSettings)
+{
+	// Respond to device added event here
+
+	Super::OnDeviceAdded(InDeviceGuid, InSettings);
+}
+
+void UExampleNetworkIngestDevice::OnDeviceRemoved()
+{
+	// Respond to device removed event here
+	//
+	// For example,execute disconnect on the connection capability:
+	// 
+	// ILiveLinkDeviceCapability_Connection::Execute_Disconnect(this);
+
+	Super::OnDeviceRemoved();
+}
+
+void UExampleNetworkIngestDevice::OnSettingChanged(const FPropertyChangedEvent& InPropertyChangedEvent)
+{
+	// Respond to any settings changes here	
+
+	Super::OnSettingChanged(InPropertyChangedEvent);
+}
+
+FString UExampleNetworkIngestDevice::GetFullTakePath(UE::CaptureManager::FTakeId InTakeId) const
+{
+	// Get the full file path to the take on disk here
+	//   See UGenericLiveLinkCPSDevice::GetFullTakePath
+
+	return FString();
+}
+
+void UExampleNetworkIngestDevice::UpdateTakeList_Implementation()
+{
+	// Query your device for a list of takes
+	// Populate an FTakeMetadata object per take
+	// Call IngestCapability->AddTake
+
+	// For example:
+	//
+	// TArray<FTakeMetadata> Takes = Device.Takes();
+	// for (const FTakeMetadata& Take : Takes)
+	// {
+	// 	 int32 TakeId = AddTake(Take);
+	// 	 TakesById[TakeId] = Take;
+	// }
+
+	// See UGenericLiveLinkCPSDevice::UpdateTakeList_Implementation
+}
+
+void UExampleNetworkIngestDevice::RunIngestTake(int32 InTakeId, const UIngestCapability_Options* InIngestOptions)
+{
+	// Get the take corresponding to the InTakeId from your device (i.e. Download the take from the device)
+	//     See UGenericLiveLinkCPSDevice::RunIngestTake
+	
+	static constexpr uint32 NumberOfTasks = 3; // Download, Convert, Upload
+	// Download - handled by you
+	// Convert + Upload - handled by UBaseIngestLiveLinkDevice
+
+	using namespace UE::CaptureManager;
+	TSharedPtr<FTaskProgress> TaskProgress = MakeShared<FTaskProgress>
+		(NumberOfTasks, FTaskProgress::FProgressReporter::CreateLambda([this, TakeId = InTakeId](double InProgress)
+			{
+				ExecuteProcessProgressReporter(TakeId, InProgress);
+			}));
+
+	// Note: IngestTake uses GetFullTakePath() so make sure you have implemented it
+
+	Super::IngestTake(InTakeId, InIngestOptions, MoveTemp(TaskProgress));
+}
+
+void UExampleNetworkIngestDevice::CancelIngest(int32 InTakeId)
+{
+	// Stop your device from sending data here
+
+	Super::CancelIngest(InTakeId);
+}
+
+ELiveLinkDeviceConnectionStatus UExampleNetworkIngestDevice::GetConnectionStatus_Implementation() const
+{
+	// Query your device for its connection status here
+
+	return ELiveLinkDeviceConnectionStatus::Disconnected;
+}
+
+FString UExampleNetworkIngestDevice::GetHardwareId_Implementation() const
+{
+	return TEXT("Example Device ID");
+}
+
+bool UExampleNetworkIngestDevice::SetHardwareId_Implementation(const FString& HardwareID)
+{
+	return false;
+}
+
+bool UExampleNetworkIngestDevice::Connect_Implementation()
+{
+	// Connect your device here
+
+	return true;
+}
+
+bool UExampleNetworkIngestDevice::Disconnect_Implementation()
+{
+	// Disconnect your device here
+	
+	return true;
+}
+
+bool UExampleNetworkIngestDevice::StartRecording_Implementation()
+{
+	// Start recording on your device here
+
+	// You can get the current slate and take information from the live link recording session
+	// 
+	// Example:
+	// 
+	// ILiveLinkRecordingSessionInfo& SessionInfo = ILiveLinkRecordingSessionInfo::Get();
+	// Device->StartRecording(SessionInfo);
+
+	return false;
+}
+
+bool UExampleNetworkIngestDevice::StopRecording_Implementation()
+{
+	// Stop recording on your device here
+
+	return false;
+}
+
+bool UExampleNetworkIngestDevice::IsRecording_Implementation() const
+{
+	// Query your device for its recording state here
+
+	return false;
+}
diff --git a/Engine/Plugins/VirtualProduction/CaptureManager/CaptureManagerApp/Source/ExampleLiveLinkDevices/Private/NetworkIngest/ExampleNetworkIngestDevice.h b/Engine/Plugins/VirtualProduction/CaptureManager/CaptureManagerApp/Source/ExampleLiveLinkDevices/Private/NetworkIngest/ExampleNetworkIngestDevice.h
new file mode 100644
--- /dev/null
+++ b/Engine/Plugins/VirtualProduction/CaptureManager/CaptureManagerApp/Source/ExampleLiveLinkDevices/Private/NetworkIngest/ExampleNetworkIngestDevice.h
@@ -0,0 +1,78 @@
+// Copyright Epic Games, Inc. All Rights Reserved.
+
+#pragma once
+
+#include "BaseIngestLiveLinkDevice.h"
+
+#include "Ingest/LiveLinkDeviceCapability_Ingest.h"
+#include "LiveLinkDeviceCapability_Connection.h"
+#include "LiveLinkDeviceCapability_Recording.h"
+
+#include "Internationalization/Text.h"
+#include "Templates/SubclassOf.h"
+
+#include "ExampleNetworkIngestDevice.generated.h"
+
+
+UCLASS()
+class UExampleNetworkIngestDeviceSettings : public ULiveLinkDeviceSettings
+{
+	GENERATED_BODY()
+
+public:
+	UPROPERTY(EditAnywhere, Category = "Example Network Ingest Device")
+	FString DisplayName = TEXT("Example Network Ingest Device");
+	
+	UPROPERTY(EditAnywhere, Category = "Example Network Ingest Device")
+	FString IpAddress;
+	
+	UPROPERTY(EditAnywhere, Category = "Example Network Ingest Device")
+	uint16 Port = 14785;
+};
+
+
+UCLASS()
+class UExampleNetworkIngestDevice : public UBaseIngestLiveLinkDevice
+	, public ILiveLinkDeviceCapability_Connection
+	, public ILiveLinkDeviceCapability_Recording
+{
+	GENERATED_BODY()
+
+public:
+	const UExampleNetworkIngestDeviceSettings* GetSettings() const;
+
+	//~ Begin ULiveLinkDevice interface
+	virtual TSubclassOf<ULiveLinkDeviceSettings> StaticGetSettingsClass() const override;
+	virtual FText GetDisplayName() const override;
+	virtual EDeviceHealth GetDeviceHealth() const override;
+	virtual FText GetHealthText() const override;
+
+	virtual void OnDeviceAdded(const FGuid InDeviceGuid, ULiveLinkDeviceSettings* InSettings) override;
+	virtual void OnDeviceRemoved() override;
+	virtual void OnSettingChanged(const FPropertyChangedEvent& InPropertyChangedEvent) override;
+	//~ End ULiveLinkDevice interface
+
+private:
+	virtual FString GetFullTakePath(UE::CaptureManager::FTakeId InTakeId) const override;
+
+	//~ Begin ULiveLinkDeviceCapability_Ingest interface
+	virtual void UpdateTakeList_Implementation() override;
+	virtual void RunIngestTake(int32 InTakeId,
+									const UIngestCapability_Options* InIngestOptions) override;
+	virtual void CancelIngest(int32 InTakeId) override;
+	//~ End ULiveLinkDeviceCapability_Ingest interface
+
+	//~ Begin ILiveLinkDeviceCapability_Connection interface
+	virtual ELiveLinkDeviceConnectionStatus GetConnectionStatus_Implementation() const override;
+	virtual FString GetHardwareId_Implementation() const override;
+	virtual bool SetHardwareId_Implementation(const FString& HardwareID) override;
+	virtual bool Connect_Implementation() override;
+	virtual bool Disconnect_Implementation() override;
+	//~ End ILiveLinkDeviceCapability_Connection interface
+
+	//~ Begin ILiveLinkDeviceCapability_Recording interface
+	virtual bool StartRecording_Implementation() override;
+	virtual bool StopRecording_Implementation() override;
+	virtual bool IsRecording_Implementation() const override;
+	//~ End ILiveLinkDeviceCapability_Recording interface
+};
