Added selection changed delagate to live link devices list and remove devices from capture manager take browser

diff --git a/Engine/Plugins/Animation/LiveLinkDevice/Source/LiveLinkDevice/Private/LiveLinkDeviceModule.cpp b/Engine/Plugins/Animation/LiveLinkDevice/Source/LiveLinkDevice/Private/LiveLinkDeviceModule.cpp
--- a/Engine/Plugins/Animation/LiveLinkDevice/Source/LiveLinkDevice/Private/LiveLinkDeviceModule.cpp
+++ b/Engine/Plugins/Animation/LiveLinkDevice/Source/LiveLinkDevice/Private/LiveLinkDeviceModule.cpp
@@ -62,6 +62,7 @@ void FLiveLinkDeviceModule::StartupModule()
 
 void FLiveLinkDeviceModule::ShutdownModule()
 {
+	OnDeviceSelectionChangedDelegate.RemoveAll(this);
 }
 
 
@@ -73,6 +74,8 @@ void FLiveLinkDeviceModule::DeviceSelectionChanged(ULiveLinkDevice* InSelectedDe
 	{
 		DetailsView->SetObject(InSelectedDevice ? InSelectedDevice->GetBaseSettings() : nullptr);
 	}
+
+	OnDeviceSelectionChangedDelegate.Broadcast(InSelectedDevice);
 }
 
 
@@ -139,6 +142,9 @@ TSharedRef<SDockTab> FLiveLinkDeviceModule::OnSpawnDevicesTab(const FSpawnTabArg
 	TSharedRef<SDockTab> DockTab = SNew(SDockTab)
 		.TabRole(ETabRole::NomadTab);
 
+	DeviceTable = SNew(SLiveLinkDeviceTable, DockTab)
+		.OnSelectionChanged_Raw(this, &FLiveLinkDeviceModule::DeviceSelectionChanged);
+
 	DockTab->SetContent(
 		SNew(SVerticalBox)
 		+ SVerticalBox::Slot()
@@ -179,7 +185,7 @@ TSharedRef<SDockTab> FLiveLinkDeviceModule::OnSpawnDevicesTab(const FSpawnTabArg
 		]
 		+ SVerticalBox::Slot()
 		[
-			SNew(SLiveLinkDeviceTable, DockTab)
+			DeviceTable.ToSharedRef()
 		]
 	);
 
diff --git a/Engine/Plugins/Animation/LiveLinkDevice/Source/LiveLinkDevice/Private/LiveLinkDeviceModule.h b/Engine/Plugins/Animation/LiveLinkDevice/Source/LiveLinkDevice/Private/LiveLinkDeviceModule.h
--- a/Engine/Plugins/Animation/LiveLinkDevice/Source/LiveLinkDevice/Private/LiveLinkDeviceModule.h
+++ b/Engine/Plugins/Animation/LiveLinkDevice/Source/LiveLinkDevice/Private/LiveLinkDeviceModule.h
@@ -2,8 +2,8 @@
 
 #pragma once
 
-#include "Logging/LogMacros.h"
-#include "Modules/ModuleManager.h"
+#include "ILiveLinkDeviceModule.h"
+
 #include "Templates/SharedPointer.h"
 #include "UObject/WeakObjectPtr.h"
 
@@ -12,37 +12,22 @@ class FLiveLinkHubDeviceComponent;
 class FSpawnTabArgs;
 class IDetailsView;
 class SDockTab;
+class SLiveLinkDeviceTable;
 class SWidget;
-class ULiveLinkDevice;
-
 
-DECLARE_LOG_CATEGORY_EXTERN(LogLiveLinkDevice, Log, All);
 
-
-class FLiveLinkDeviceModule : public IModuleInterface
+class FLiveLinkDeviceModule : public ILiveLinkDeviceModule
 {
 public:
 	static const FName DevicesTabName;
 	static const FName DeviceDetailsTabName;
 
-	/**
-	 * Singleton-like access to this module's interface.  This is just for convenience!
-	 * Beware of calling this during the shutdown phase, though.  Your module might have been unloaded already.
-	 *
-	 * @return Returns singleton instance, loading the module on demand if needed
-	 */
-	static FLiveLinkDeviceModule& Get()
-	{
-		static const FName ModuleName = "LiveLinkDevice";
-		return FModuleManager::LoadModuleChecked<FLiveLinkDeviceModule>(ModuleName);
-	}
-
 	//~ Begin IModuleInterface
 	virtual void StartupModule() override;
 	virtual void ShutdownModule() override;
 	//~ End IModuleInterface
 
-	void DeviceSelectionChanged(ULiveLinkDevice* InSelectedDevice);
+	FOnDeviceSelectionChangedDelegate& OnSelectionChanged() override { return OnDeviceSelectionChangedDelegate; }
 
 private:
 	TSharedRef<SWidget> OnGenerateAddDeviceMenu();
@@ -50,7 +35,12 @@ private:
 	TSharedRef<SDockTab> OnSpawnDeviceDetailsTab(const FSpawnTabArgs& InSpawnTabArgs);
 
 private:
+	void DeviceSelectionChanged(ULiveLinkDevice* InSelectedDevice);
+
 	TSharedPtr<FLiveLinkHubDeviceComponent> LLHComponent;
+	TSharedPtr<SLiveLinkDeviceTable> DeviceTable;
 	TSharedPtr<IDetailsView> DetailsView;
 	TWeakObjectPtr<ULiveLinkDevice> WeakSelectedDevice;
+
+	FOnDeviceSelectionChangedDelegate OnDeviceSelectionChangedDelegate;
 };
diff --git a/Engine/Plugins/Animation/LiveLinkDevice/Source/LiveLinkDevice/Private/Widgets/DeviceTable.cpp b/Engine/Plugins/Animation/LiveLinkDevice/Source/LiveLinkDevice/Private/Widgets/DeviceTable.cpp
--- a/Engine/Plugins/Animation/LiveLinkDevice/Source/LiveLinkDevice/Private/Widgets/DeviceTable.cpp
+++ b/Engine/Plugins/Animation/LiveLinkDevice/Source/LiveLinkDevice/Private/Widgets/DeviceTable.cpp
@@ -29,6 +29,8 @@ void SLiveLinkDeviceTable::Construct(const FArguments& InArgs, TSharedRef<SDockT
 {
 	WeakContainingTab = InContainingTab;
 
+	OnSelectionChangedDelegate = InArgs._OnSelectionChanged;
+
 	ULiveLinkDeviceSubsystem* Subsystem = GEngine->GetEngineSubsystem<ULiveLinkDeviceSubsystem>();
 
 	ChildSlot
@@ -37,7 +39,7 @@ void SLiveLinkDeviceTable::Construct(const FArguments& InArgs, TSharedRef<SDockT
 		.SelectionMode(ESelectionMode::Single)
 		.ListItemsSource(&DeviceRows)
 		.OnSelectionChanged_Lambda(
-			[]
+			[this]
 			(TSharedPtr<FLiveLinkDeviceRowData> InSelectedRow, ESelectInfo::Type InSelectType)
 			{
 				ULiveLinkDevice* Device = nullptr;
@@ -46,8 +48,7 @@ void SLiveLinkDeviceTable::Construct(const FArguments& InArgs, TSharedRef<SDockT
 					Device = InSelectedRow->WeakDevice.Get();
 				}
 
-				FLiveLinkDeviceModule& Module = FLiveLinkDeviceModule::Get();
-				Module.DeviceSelectionChanged(Device);
+				OnSelectionChangedDelegate.ExecuteIfBound(Device);
 			}
 		)
 		.OnMouseButtonDoubleClick_Static(
@@ -83,10 +84,13 @@ void SLiveLinkDeviceTable::Construct(const FArguments& InArgs, TSharedRef<SDockT
 
 SLiveLinkDeviceTable::~SLiveLinkDeviceTable()
 {
-	if (ULiveLinkDeviceSubsystem* Subsystem = GEngine->GetEngineSubsystem<ULiveLinkDeviceSubsystem>())
+	if (GEngine)
 	{
-		Subsystem->OnDeviceAdded().RemoveAll(this);
-		Subsystem->OnDeviceRemoved().RemoveAll(this);
+		if (ULiveLinkDeviceSubsystem* Subsystem = GEngine->GetEngineSubsystem<ULiveLinkDeviceSubsystem>())
+		{
+			Subsystem->OnDeviceAdded().RemoveAll(this);
+			Subsystem->OnDeviceRemoved().RemoveAll(this);
+		}
 	}
 }
 
diff --git a/Engine/Plugins/Animation/LiveLinkDevice/Source/LiveLinkDevice/Private/Widgets/DeviceTable.h b/Engine/Plugins/Animation/LiveLinkDevice/Source/LiveLinkDevice/Private/Widgets/DeviceTable.h
--- a/Engine/Plugins/Animation/LiveLinkDevice/Source/LiveLinkDevice/Private/Widgets/DeviceTable.h
+++ b/Engine/Plugins/Animation/LiveLinkDevice/Source/LiveLinkDevice/Private/Widgets/DeviceTable.h
@@ -13,6 +13,9 @@ class ULiveLinkDevice;
 class ULiveLinkDeviceCapability;
 
 
+DECLARE_DELEGATE_OneParam(FOnSelectionChangedDelegate, ULiveLinkDevice*);
+
+
 struct FLiveLinkDeviceRowData
 {
 	TWeakObjectPtr<ULiveLinkDevice> WeakDevice;
@@ -24,6 +27,7 @@ struct FLiveLinkDeviceRowData
 class SLiveLinkDeviceTable : public SCompoundWidget
 {
 	SLATE_BEGIN_ARGS(SLiveLinkDeviceTable) {}
+		SLATE_EVENT(FOnSelectionChangedDelegate, OnSelectionChanged)
 	SLATE_END_ARGS()
 
 public:
@@ -46,6 +50,8 @@ private:
 	TSharedPtr<SHeaderRow> HeaderRow;
 	TArray<TSharedPtr<FLiveLinkDeviceRowData>> DeviceRows;
 	TSharedPtr<SListView<TSharedPtr<FLiveLinkDeviceRowData>>> DeviceListView;
+
+	FOnSelectionChangedDelegate OnSelectionChangedDelegate;
 };
 
 
diff --git a/Engine/Plugins/Animation/LiveLinkDevice/Source/LiveLinkDevice/Public/ILiveLinkDeviceModule.h b/Engine/Plugins/Animation/LiveLinkDevice/Source/LiveLinkDevice/Public/ILiveLinkDeviceModule.h
new file mode 100644
--- /dev/null
+++ b/Engine/Plugins/Animation/LiveLinkDevice/Source/LiveLinkDevice/Public/ILiveLinkDeviceModule.h
@@ -0,0 +1,39 @@
+// Copyright Epic Games, Inc. All Rights Reserved.
+
+#pragma once
+
+#include "Logging/LogMacros.h"
+#include "Modules/ModuleManager.h"
+
+class ULiveLinkDevice;
+
+DECLARE_MULTICAST_DELEGATE_OneParam(FOnDeviceSelectionChangedDelegate, ULiveLinkDevice*);
+
+DECLARE_LOG_CATEGORY_EXTERN(LogLiveLinkDevice, Log, All);
+
+/**
+* The public interface to this module
+*/
+class ILiveLinkDeviceModule : public IModuleInterface
+{
+
+public:
+
+	/**
+	* Singleton-like access to this module's interface.  This is just for convenience!
+	* Beware of calling this during the shutdown phase, though.  Your module might have been unloaded already.
+	*
+	* @return Returns singleton instance, loading the module on demand if needed
+	*/
+	static inline ILiveLinkDeviceModule& Get()
+	{
+		static const FName ModuleName = "LiveLinkDevice";
+		return FModuleManager::LoadModuleChecked<ILiveLinkDeviceModule>(ModuleName);
+	}
+
+	/**
+	* Gets the device selection changed delegate
+	* @return Returns the device selection changed delegate
+	*/
+	virtual FOnDeviceSelectionChangedDelegate& OnSelectionChanged() = 0;
+};
diff --git a/Engine/Plugins/VirtualProduction/CaptureManager/CaptureManagerApp/Source/CaptureManagerEditor/Private/UI/TakesView.cpp b/Engine/Plugins/VirtualProduction/CaptureManager/CaptureManagerApp/Source/CaptureManagerEditor/Private/UI/TakesView.cpp
--- a/Engine/Plugins/VirtualProduction/CaptureManager/CaptureManagerApp/Source/CaptureManagerEditor/Private/UI/TakesView.cpp
+++ b/Engine/Plugins/VirtualProduction/CaptureManager/CaptureManagerApp/Source/CaptureManagerEditor/Private/UI/TakesView.cpp
@@ -10,6 +10,7 @@
 #include "ContentBrowserModule.h"
 #include "Async/Monitor.h"
 #include "Misc/FileHelper.h"
+#include "ILiveLinkDeviceModule.h"
 
 #include "Widgets/Input/SButton.h"
 #include "Widgets/Input/SComboBox.h"
@@ -29,7 +30,7 @@ static const FString OuterPackageName = TEXT("CaptureDevices");
 STakesView::STakesView() = default;
 
 void STakesView::Construct(
-	const FArguments& InArgs, 
+	const FArguments& InArgs,
 	TSharedRef<UE::CaptureManager::FIngestJobSettingsManager> InIngestJobsSettingsManager,
 	TUniquePtr<FGetCurrentSelectionDelegate> InGetCurrentSelectionDelegate
 )
@@ -56,7 +57,6 @@ void STakesView::Construct(
 
 	IContentBrowserSingleton& ContentBrowser = FModuleManager::LoadModuleChecked<FContentBrowserModule>("ContentBrowser").Get();
 
-	FAssetPickerConfig DevicesPickerConfig = InArgs._DevicesPickerConfig;
 	FAssetPickerConfig TakesPickerConfig = InArgs._TakesPickerConfig;
 
 	ChildSlot
@@ -64,18 +64,7 @@ void STakesView::Construct(
 			SNew(SVerticalBox)
 				+ SVerticalBox::Slot()
 				[
-					SNew(SSplitter)
-						.Orientation(Orient_Horizontal)
-						+ SSplitter::Slot()
-						.Value(0.25f)
-						[
-							ContentBrowser.CreateAssetPicker(DevicesPickerConfig)
-						]
-						+ SSplitter::Slot()
-						.Value(0.75f)
-						[
-							ContentBrowser.CreateAssetPicker(TakesPickerConfig)
-						]
+					ContentBrowser.CreateAssetPicker(TakesPickerConfig)
 				]
 				+ SVerticalBox::Slot()
 				.VAlign(EVerticalAlignment::VAlign_Bottom)
@@ -212,16 +201,16 @@ TArray<FAssetData> FTakesView::CreateTakeAssetsDataForCaptureDevice(TOptional<FS
 	TSharedPtr<FCaptureManagerPanelController> Controller = PanelController.Pin();
 	if (Controller)
 	{
-		for (const TPair<FGuid, TArray<TStrongObjectPtr<UTakeVirtualAsset>>>& CSIdTakeTuple : CaptureDeviceTakesMap)
+		for (const TPair<FGuid, TArray<TStrongObjectPtr<UTakeVirtualAsset>>>& DeviceIdTakeTuple : CaptureDeviceTakesMap)
 		{
-			FGuid CaptureDeviceId = CSIdTakeTuple.Key;
+			FGuid CaptureDeviceId = DeviceIdTakeTuple.Key;
 			TObjectPtr<ULiveLinkDevice> Device = Controller->GetCaptureDevice(CaptureDeviceId);
 
 			if (!InCaptureDeviceName.IsSet() || Device.GetName() == InCaptureDeviceName.GetValue())
 			{
 				TArray<TStrongObjectPtr<UTakeVirtualAsset>> Takes;
 
-				for (const TStrongObjectPtr<UTakeVirtualAsset>& Take : CSIdTakeTuple.Value)
+				for (const TStrongObjectPtr<UTakeVirtualAsset>& Take : DeviceIdTakeTuple.Value)
 				{
 					Takes.Add(Take);
 				}
@@ -249,7 +238,7 @@ TArray<FAssetData> FTakesView::CreateTakeAssetsData(const TMap<FString, TArray<T
 	return TakeAssetsData;
 }
 
-FTakesView::FTakesView(TWeakPtr<FCaptureManagerPanelController> InController) : 
+FTakesView::FTakesView(TWeakPtr<FCaptureManagerPanelController> InController) :
 	PanelController(InController)
 {
 	CreateTakesView();
@@ -314,59 +303,12 @@ void FTakesView::CreateTakesView()
 
 	TSharedRef<UE::CaptureManager::FIngestJobSettingsManager> IngestJobSettingsManager = Controller->GetIngestJobSettingsManager();
 
-	FAssetPickerConfig DevicesPickerConfig;
-	{
-		DevicesPickerConfig.SelectionMode = ESelectionMode::Single;
-		DevicesPickerConfig.InitialAssetViewType = EAssetViewType::Column;
-		DevicesPickerConfig.bFocusSearchBoxWhenOpened = true;
-		DevicesPickerConfig.bAllowNullSelection = false;
-		DevicesPickerConfig.bShowBottomToolbar = true;
-		DevicesPickerConfig.bAutohideSearchBar = false;
-		DevicesPickerConfig.bAllowDragging = false;
-		DevicesPickerConfig.bCanShowClasses = false;
-		DevicesPickerConfig.bShowPathInColumnView = false;
-		DevicesPickerConfig.bSortByPathInColumnView = false;
-		DevicesPickerConfig.bShowTypeInColumnView = false;
-		DevicesPickerConfig.bForceShowEngineContent = true;
-		DevicesPickerConfig.bForceShowPluginContent = true;
-		DevicesPickerConfig.SyncToAssetsDelegates.Add(&SyncToDevicesDelegate);
-		DevicesPickerConfig.RefreshAssetViewDelegates.Add(&RefreshDevicesViewDelegate);
-		DevicesPickerConfig.Filter.ClassPaths.Add(UCaptureSourceVirtualAsset::StaticClass()->GetClassPathName());
-		DevicesPickerConfig.Filter.bRecursiveClasses = true;
-		DevicesPickerConfig.Filter.bRecursivePaths = true;
-
-		DevicesPickerConfig.OnGetCustomSourceAssets =
-			FOnGetCustomSourceAssets::CreateLambda([this](const FARFilter& SourceFilter, TArray<FAssetData>& OutAssets)
-				{
-					for (const TStrongObjectPtr<UCaptureSourceVirtualAsset>& CaptureDevice : CaptureDevices)
-					{
-						FAssetData AssetData(CaptureDevice.Get());
-						OutAssets.Add(MoveTemp(AssetData));
-					}
-				});
-
-		DevicesPickerConfig.OnAssetSelected = FOnAssetSelected::CreateLambda([this](const FAssetData& InSelectedItemData)
-			{
-				if (!InSelectedItemData.GetAsset())
-				{
-					SelectedItem = nullptr;
-					RefreshTakesViewDelegate.ExecuteIfBound(true);
-					return;
-				}
-
-				TStrongObjectPtr<UCaptureSourceVirtualAsset>* CaptureDevice = 
-					CaptureDevices.FindByPredicate([InSelectedItemData](const TStrongObjectPtr<UCaptureSourceVirtualAsset>& InCaptureDevice)
-				{
-					return InCaptureDevice->Id == StaticCast<UCaptureSourceVirtualAsset*>(InSelectedItemData.GetAsset())->Id;
-				});
-
-				if (CaptureDevice)
-				{
-					SelectedItem = CaptureDevice->Get();
-					RefreshTakesViewDelegate.ExecuteIfBound(true);
-				}
-			});
-	}
+	ILiveLinkDeviceModule& Module = ILiveLinkDeviceModule::Get();
+	Module.OnSelectionChanged().AddLambda([this](ULiveLinkDevice* InSelectedDevice)
+		{
+			SelectedItem = InSelectedDevice;
+			RefreshTakesViewDelegate.ExecuteIfBound(true);
+		});
 
 	TUniquePtr<FGetCurrentSelectionDelegate> GetCurrentSelectionDelegate = MakeUnique<FGetCurrentSelectionDelegate>();
 
@@ -406,180 +348,156 @@ void FTakesView::CreateTakesView()
 
 	SAssignNew(TakesTileView, STakesView, MoveTemp(IngestJobSettingsManager), MoveTemp(GetCurrentSelectionDelegate))
 		.OnAddTakesToIngestQueue_Raw(this, &FTakesView::AddTakesToIngestQueue)
-		.DevicesPickerConfig(DevicesPickerConfig)
 		.TakesPickerConfig(TakesPickerConfig);
 }
 
 void FTakesView::CaptureDeviceStarted(FGuid InCaptureDeviceId)
 {
 	AsyncTask(ENamedThreads::GameThread, [this, Controller = PanelController.Pin(), InCaptureDeviceId]()
-	{
-		if (!Controller)
 		{
-			return;
-		}
-
-		using namespace UE::CaptureManager;
-
-		TObjectPtr<ULiveLinkDevice> Device = Controller->GetCaptureDevice(InCaptureDeviceId);
-
-		if (!Device)
-		{
-			return;
-		}
-
-		if (!Device->Implements<ULiveLinkDeviceCapability_Ingest>())
-		{
-			return;
-		}
-
-		FString PackageName = GetTransientPackage()->GetPathName() / OuterPackageName;
-		UPackage* Package = CreatePackage(*PackageName);
-
-		FString CaptureDeviceName = Device->GetName();
-		FName Name(CaptureDeviceName);
-
-		if (UCaptureSourceVirtualAsset* FoundObject = FindObject<UCaptureSourceVirtualAsset>(Package, *CaptureDeviceName, true))
-		{
-			Name = MakeUniqueObjectName(Package, UCaptureSourceVirtualAsset::StaticClass(), FName(CaptureDeviceName));
-		}
-		
-		TStrongObjectPtr<UCaptureSourceVirtualAsset> CaptureDeviceObjectPtr(NewObject<UCaptureSourceVirtualAsset>(Package, Name));
-		CaptureDeviceObjectPtr->Id = InCaptureDeviceId;
-
-		CaptureDevices.AddUnique(MoveTemp(CaptureDeviceObjectPtr));
-
-		FIngestUpdateTakeListCallback Callback = FIngestUpdateTakeListCallback::Type::CreateLambda(
-			[this, InCaptureDeviceId](const TArray<FTakeId>& InTakeIds)
+			if (!Controller)
 			{
-				CaptureDeviceTakesMap.Remove(InCaptureDeviceId);
+				return;
+			}
 
-				TArray<TStrongObjectPtr<UTakeVirtualAsset>> TakeObjects = MakeTakeObjects(InCaptureDeviceId, InTakeIds);
+			using namespace UE::CaptureManager;
 
-				CaptureDeviceTakesMap.Add(InCaptureDeviceId, MoveTemp(TakeObjects));
+			TObjectPtr<ULiveLinkDevice> Device = Controller->GetCaptureDevice(InCaptureDeviceId);
 
-				RefreshDevicesViewDelegate.ExecuteIfBound(true);
-				RefreshTakesViewDelegate.ExecuteIfBound(true);
+			if (!Device)
+			{
+				return;
 			}
-		);
 
-		TScriptInterface<ILiveLinkDeviceCapability_Ingest> CapabilityInterface(Device);
-
-		CapabilityInterface->SubscribeToEvent(FString(FIngestCapability_TakeAddedEvent::Name),
-											  FCaptureEventHandler([this, InCaptureDeviceId](TSharedPtr<const FCaptureEvent> InTakeAddedEvent)
-		{
-			check(IsInGameThread());
+			if (!Device->Implements<ULiveLinkDeviceCapability_Ingest>())
+			{
+				return;
+			}
 
-			TSharedPtr<const FIngestCapability_TakeAddedEvent> Event = 
-				StaticCastSharedPtr<const FIngestCapability_TakeAddedEvent>(InTakeAddedEvent);
+			FString PackageName = GetTransientPackage()->GetPathName() / OuterPackageName;
+			UPackage* Package = CreatePackage(*PackageName);
 
-			TArray<FTakeId> TakeIds = { Event->TakeId };
-			TArray<TStrongObjectPtr<UTakeVirtualAsset>> NewTakeObject = MakeTakeObjects(InCaptureDeviceId, { TakeIds });
+			FString CaptureDeviceName = Device->GetName();
+			FName Name(CaptureDeviceName);
 
-			if (TArray<TStrongObjectPtr<UTakeVirtualAsset>>* TakeObjects = CaptureDeviceTakesMap.Find(InCaptureDeviceId))
-			{
-				TakeObjects->Append(NewTakeObject.GetData(), NewTakeObject.Num());
+			FIngestUpdateTakeListCallback Callback = FIngestUpdateTakeListCallback::Type::CreateLambda(
+				[this, InCaptureDeviceId](const TArray<FTakeId>& InTakeIds)
+				{
+					CaptureDeviceTakesMap.Remove(InCaptureDeviceId);
 
-				RefreshTakesViewDelegate.ExecuteIfBound(true);
-			}
-		}));
+					TArray<TStrongObjectPtr<UTakeVirtualAsset>> TakeObjects = MakeTakeObjects(InCaptureDeviceId, InTakeIds);
 
-		CapabilityInterface->SubscribeToEvent(FString(FIngestCapability_TakeUpdatedEvent::Name),
-											  FCaptureEventHandler([this, InCaptureDeviceId](TSharedPtr<const FCaptureEvent> InTakeUpdatedEvent)
-		{
-			check(IsInGameThread());
+					CaptureDeviceTakesMap.Add(InCaptureDeviceId, MoveTemp(TakeObjects));
 
-			TSharedPtr<const FIngestCapability_TakeUpdatedEvent> Event = 
-				StaticCastSharedPtr<const FIngestCapability_TakeUpdatedEvent>(InTakeUpdatedEvent);
+					RefreshTakesViewDelegate.ExecuteIfBound(true);
+				}
+			);
 
-			TArray<FTakeId> TakeIds = { Event->TakeId };
-			TArray<TStrongObjectPtr<UTakeVirtualAsset>> NewTakeObject = MakeTakeObjects(InCaptureDeviceId, { TakeIds });
-			check(NewTakeObject.Num() == 1);
+			TScriptInterface<ILiveLinkDeviceCapability_Ingest> CapabilityInterface(Device);
 
-			if (TArray<TStrongObjectPtr<UTakeVirtualAsset>>* TakeObjects = CaptureDeviceTakesMap.Find(InCaptureDeviceId))
-			{
-				for (TStrongObjectPtr<UTakeVirtualAsset>& TakeVirtualAsset : *TakeObjects)
-				{
-					if (TakeVirtualAsset->TakeId == Event->TakeId)
+			CapabilityInterface->SubscribeToEvent(FString(FIngestCapability_TakeAddedEvent::Name),
+				FCaptureEventHandler([this, InCaptureDeviceId](TSharedPtr<const FCaptureEvent> InTakeAddedEvent)
 					{
-						TakeVirtualAsset = NewTakeObject[0];
-						break;
-					}
-				}
+						check(IsInGameThread());
 
-				RefreshTakesViewDelegate.ExecuteIfBound(true);
-			}
-		}));
+						TSharedPtr<const FIngestCapability_TakeAddedEvent> Event =
+							StaticCastSharedPtr<const FIngestCapability_TakeAddedEvent>(InTakeAddedEvent);
 
-		CapabilityInterface->SubscribeToEvent(FString(FIngestCapability_TakeRemovedEvent::Name),
-											  FCaptureEventHandler([this, InCaptureDeviceId](TSharedPtr<const FCaptureEvent> InTakeRemovedEvent)
-		{
-			check(IsInGameThread());
+						TArray<FTakeId> TakeIds = { Event->TakeId };
+						TArray<TStrongObjectPtr<UTakeVirtualAsset>> NewTakeObject = MakeTakeObjects(InCaptureDeviceId, { TakeIds });
 
-			TSharedPtr<const FIngestCapability_TakeRemovedEvent> Event = 
-				StaticCastSharedPtr<const FIngestCapability_TakeRemovedEvent>(InTakeRemovedEvent);
+						if (TArray<TStrongObjectPtr<UTakeVirtualAsset>>* TakeObjects = CaptureDeviceTakesMap.Find(InCaptureDeviceId))
+						{
+							TakeObjects->Append(NewTakeObject.GetData(), NewTakeObject.Num());
 
-			if (TArray<TStrongObjectPtr<UTakeVirtualAsset>>* TakeObjects = CaptureDeviceTakesMap.Find(InCaptureDeviceId))
-			{
-				int32 NumRemoved = TakeObjects->RemoveAll([this, TakeId = Event->TakeId](const TStrongObjectPtr<UTakeVirtualAsset>& InObject)
-				{
-					return InObject->TakeId == TakeId;
-				});
+							RefreshTakesViewDelegate.ExecuteIfBound(true);
+						}
+					}));
 
-				if (NumRemoved > 0)
-				{
-					RefreshTakesViewDelegate.ExecuteIfBound(true);
-				}
-			}
-		}));
-		
-		ILiveLinkDeviceCapability_Ingest::Execute_CreateUpdateTakeListCallback(Device)->Callback = MoveTemp(Callback);
-		ILiveLinkDeviceCapability_Ingest::Execute_UpdateTakeList(Device);
-	});
+			CapabilityInterface->SubscribeToEvent(FString(FIngestCapability_TakeUpdatedEvent::Name),
+				FCaptureEventHandler([this, InCaptureDeviceId](TSharedPtr<const FCaptureEvent> InTakeUpdatedEvent)
+					{
+						check(IsInGameThread());
+
+						TSharedPtr<const FIngestCapability_TakeUpdatedEvent> Event =
+							StaticCastSharedPtr<const FIngestCapability_TakeUpdatedEvent>(InTakeUpdatedEvent);
+
+						TArray<FTakeId> TakeIds = { Event->TakeId };
+						TArray<TStrongObjectPtr<UTakeVirtualAsset>> NewTakeObject = MakeTakeObjects(InCaptureDeviceId, { TakeIds });
+						check(NewTakeObject.Num() == 1);
+
+						if (TArray<TStrongObjectPtr<UTakeVirtualAsset>>* TakeObjects = CaptureDeviceTakesMap.Find(InCaptureDeviceId))
+						{
+							for (TStrongObjectPtr<UTakeVirtualAsset>& TakeVirtualAsset : *TakeObjects)
+							{
+								if (TakeVirtualAsset->TakeId == Event->TakeId)
+								{
+									TakeVirtualAsset = NewTakeObject[0];
+									break;
+								}
+							}
+
+							RefreshTakesViewDelegate.ExecuteIfBound(true);
+						}
+					}));
+
+			CapabilityInterface->SubscribeToEvent(FString(FIngestCapability_TakeRemovedEvent::Name),
+				FCaptureEventHandler([this, InCaptureDeviceId](TSharedPtr<const FCaptureEvent> InTakeRemovedEvent)
+					{
+						check(IsInGameThread());
+
+						TSharedPtr<const FIngestCapability_TakeRemovedEvent> Event =
+							StaticCastSharedPtr<const FIngestCapability_TakeRemovedEvent>(InTakeRemovedEvent);
+
+						if (TArray<TStrongObjectPtr<UTakeVirtualAsset>>* TakeObjects = CaptureDeviceTakesMap.Find(InCaptureDeviceId))
+						{
+							int32 NumRemoved = TakeObjects->RemoveAll([this, TakeId = Event->TakeId](const TStrongObjectPtr<UTakeVirtualAsset>& InObject)
+								{
+									return InObject->TakeId == TakeId;
+								});
+
+							if (NumRemoved > 0)
+							{
+								RefreshTakesViewDelegate.ExecuteIfBound(true);
+							}
+						}
+					}));
+
+			ILiveLinkDeviceCapability_Ingest::Execute_CreateUpdateTakeListCallback(Device)->Callback = MoveTemp(Callback);
+			ILiveLinkDeviceCapability_Ingest::Execute_UpdateTakeList(Device);
+		});
 }
 
 void FTakesView::CaptureDeviceStopped(FGuid InCaptureDeviceId)
 {
 	AsyncTask(ENamedThreads::GameThread, [this, Controller = PanelController.Pin(), InCaptureDeviceId]()
-	{
-		if (!Controller)
 		{
-			return;
-		}
+			if (!Controller)
+			{
+				return;
+			}
 
-		using namespace UE::CaptureManager;
+			using namespace UE::CaptureManager;
 
-		TObjectPtr<ULiveLinkDevice> Device = Controller->GetCaptureDevice(InCaptureDeviceId);
+			TObjectPtr<ULiveLinkDevice> Device = Controller->GetCaptureDevice(InCaptureDeviceId);
 
-		if (!Device)
-		{
-			return;
-		}
+			if (!Device)
+			{
+				return;
+			}
 
-		CaptureDeviceTakesMap.Remove(InCaptureDeviceId);
-		CaptureDevices.RemoveAll([InCaptureDeviceId](const TStrongObjectPtr<UCaptureSourceVirtualAsset>& InCaptureDevice)
-		{
-			return InCaptureDevice->Id == InCaptureDeviceId;
+			CaptureDeviceTakesMap.Remove(InCaptureDeviceId);
+			RefreshTakesViewDelegate.ExecuteIfBound(true);
 		});
-
-		RefreshDevicesViewDelegate.ExecuteIfBound(true);
-		RefreshTakesViewDelegate.ExecuteIfBound(true);
-	});
 }
 
 void FTakesView::CaptureDeviceToBeRemoved(FGuid InCaptureDeviceId)
 {
 	AsyncTask(ENamedThreads::GameThread, [this, InCaptureDeviceId]()
-	{
-		CaptureDeviceTakesMap.Remove(InCaptureDeviceId);
-		CaptureDevices.RemoveAll([InCaptureDeviceId](const TStrongObjectPtr<UCaptureSourceVirtualAsset>& InCaptureDevice)
 		{
-			return InCaptureDevice->Id == InCaptureDeviceId;
+			CaptureDeviceTakesMap.Remove(InCaptureDeviceId);
+			RefreshTakesViewDelegate.ExecuteIfBound(true);
 		});
-
-		RefreshDevicesViewDelegate.ExecuteIfBound(true);
-		RefreshTakesViewDelegate.ExecuteIfBound(true);
-	});
 }
 
 TArray<TStrongObjectPtr<UTakeVirtualAsset>> FTakesView::MakeTakeObjects(FGuid InCaptureDeviceId, TOptional<TArray<UE::CaptureManager::FTakeId>> InTakeIds) {
diff --git a/Engine/Plugins/VirtualProduction/CaptureManager/CaptureManagerApp/Source/CaptureManagerEditor/Private/UI/TakesView.h b/Engine/Plugins/VirtualProduction/CaptureManager/CaptureManagerApp/Source/CaptureManagerEditor/Private/UI/TakesView.h
--- a/Engine/Plugins/VirtualProduction/CaptureManager/CaptureManagerApp/Source/CaptureManagerEditor/Private/UI/TakesView.h
+++ b/Engine/Plugins/VirtualProduction/CaptureManager/CaptureManagerApp/Source/CaptureManagerEditor/Private/UI/TakesView.h
@@ -21,6 +21,7 @@ class FCaptureSourcesView;
 class SEditableTextBox;
 class SSearchBox;
 class SSlateBoxBrush;
+class ULiveLinkDevice;
 
 namespace UE::CaptureManager
 {
@@ -33,7 +34,6 @@ public:
 	DECLARE_DELEGATE_TwoParams(FOnAddTakesToIngestQueue, const TArray<TStrongObjectPtr<UTakeVirtualAsset>>&, const UIngestJobSettingsBase&);
 
 	SLATE_BEGIN_ARGS(STakesView) {}
-		SLATE_ARGUMENT(FAssetPickerConfig, DevicesPickerConfig)
 		SLATE_ARGUMENT(FAssetPickerConfig, TakesPickerConfig)
 		SLATE_EVENT(FOnAddTakesToIngestQueue, OnAddTakesToIngestQueue)
 	SLATE_END_ARGS()
@@ -78,11 +78,9 @@ private:
 	TWeakPtr<FCaptureManagerPanelController> PanelController;
 
 	TMap<FGuid, TArray<TStrongObjectPtr<UTakeVirtualAsset>>> CaptureDeviceTakesMap;
-	TArray<TStrongObjectPtr<UCaptureSourceVirtualAsset>> CaptureDevices;
-	TWeakObjectPtr<UCaptureSourceVirtualAsset> SelectedItem;
+	TWeakObjectPtr<ULiveLinkDevice> SelectedItem;
 
 	FRefreshAssetViewDelegate RefreshTakesViewDelegate;
-	FRefreshAssetViewDelegate RefreshDevicesViewDelegate;
 	FSyncToAssetsDelegate SyncToDevicesDelegate;
 
 	TArray<FAssetData> CreateTakeAssetsDataForCaptureDevice(TOptional<FString> InCaptureDeviceName = TOptional<FString>());
