[CaptureManager] Add naming tokens to Capture Manager in Live Link Hub

diff --git a/Engine/Plugins/VirtualProduction/CaptureManager/CaptureManagerApp/CaptureManagerApp.uplugin b/Engine/Plugins/VirtualProduction/CaptureManager/CaptureManagerApp/CaptureManagerApp.uplugin
--- a/Engine/Plugins/VirtualProduction/CaptureManager/CaptureManagerApp/CaptureManagerApp.uplugin
+++ b/Engine/Plugins/VirtualProduction/CaptureManager/CaptureManagerApp/CaptureManagerApp.uplugin
@@ -15,7 +15,8 @@
 	"IsBetaVersion": false,
 	"IsExperimentalVersion": true,
 	"Installed": false,
-	"Hidden" :  true,
+	"Hidden": true,
+	"SupportedPrograms": [ "LiveLinkHub" ],
 	"Modules": [
 		{
 			"Name": "CaptureManagerEditor",
@@ -88,6 +89,10 @@
 			"Name": "LiveLinkHub",
 			"Enabled": true
 		},
+		{
+			"Name": "NamingTokens",
+			"Enabled": true
+		},
 		{
 			"Name": "ControlFlows",
 			"Enabled": true
diff --git a/Engine/Plugins/VirtualProduction/CaptureManager/CaptureManagerApp/Source/CaptureManagerEditor/CaptureManagerEditor.Build.cs b/Engine/Plugins/VirtualProduction/CaptureManager/CaptureManagerApp/Source/CaptureManagerEditor/CaptureManagerEditor.Build.cs
--- a/Engine/Plugins/VirtualProduction/CaptureManager/CaptureManagerApp/Source/CaptureManagerEditor/CaptureManagerEditor.Build.cs
+++ b/Engine/Plugins/VirtualProduction/CaptureManager/CaptureManagerApp/Source/CaptureManagerEditor/CaptureManagerEditor.Build.cs
@@ -37,7 +37,8 @@ public class CaptureManagerEditor : ModuleRules
 			"CaptureCommunicationManager",
 			"CaptureDataCore",
 			"CaptureManagerSettings",
-			"ImageCore"
+			"ImageCore",
+			"NamingTokens"
 		});
 	}
 }
diff --git a/Engine/Plugins/VirtualProduction/CaptureManager/CaptureManagerApp/Source/CaptureManagerEditor/Private/IngestManagement/IngestJobExecutor.cpp b/Engine/Plugins/VirtualProduction/CaptureManager/CaptureManagerApp/Source/CaptureManagerEditor/Private/IngestManagement/IngestJobExecutor.cpp
--- a/Engine/Plugins/VirtualProduction/CaptureManager/CaptureManagerApp/Source/CaptureManagerEditor/Private/IngestManagement/IngestJobExecutor.cpp
+++ b/Engine/Plugins/VirtualProduction/CaptureManager/CaptureManagerApp/Source/CaptureManagerEditor/Private/IngestManagement/IngestJobExecutor.cpp
@@ -9,6 +9,15 @@
 #include "Ingest/LiveLinkDeviceCapability_Ingest.h"
 #include "LiveLinkDeviceSubsystem.h"
 #include "LiveLinkDevice.h"
+#include "Settings/CaptureManagerSettings.h"
+
+#include "Async/HelperFunctions.h"
+
+#include "Editor.h"
+#include "GlobalNamingTokens.h"
+#include "NamingTokenData.h"
+#include "NamingTokensEditorSubsystem.h"
+#include "Settings/CaptureManagerTemplateTokens.h"
 
 DEFINE_LOG_CATEGORY_STATIC(LogCaptureManagerIngestJobExecutor, Log, All);
 
@@ -19,11 +28,11 @@ EIngestCapability_ImagePixelFormat ConvertPixelFormat(::EImagePixelFormat InImag
 {
 	switch (InImagePixelFormat)
 	{
-		case ::EImagePixelFormat::U8_Mono:
-			return EIngestCapability_ImagePixelFormat::U8_Mono;
-		case ::EImagePixelFormat::U8_BGRA:
-		default:
-			return EIngestCapability_ImagePixelFormat::U8_BGRA;
+	case ::EImagePixelFormat::U8_Mono:
+		return EIngestCapability_ImagePixelFormat::U8_Mono;
+	case ::EImagePixelFormat::U8_BGRA:
+	default:
+		return EIngestCapability_ImagePixelFormat::U8_BGRA;
 	}
 }
 
@@ -74,7 +83,7 @@ uint32 FIngestJobExecutor::Run()
 			check(!IngestProcessHandle);
 
 			IngestProcessHandle = MoveTemp(TakeProcessHandle);
-			
+
 			// There's a period between RunIngest and the assignment of the member handle, during which Stop() may have
 			// been called. If so, immediately terminate (as Cancel() will not be called otherwise).
 			if (bStopRequested && !Future.IsReady() && IngestProcessHandle)
@@ -125,6 +134,66 @@ void FIngestJobExecutor::Exit()
 
 UIngestCapability_ProcessHandle* FIngestJobExecutor::RunIngest(FIngestJob& InJob, TPromise<void>& InPromise)
 {
+	if (!GEditor)
+	{
+		UE_LOG(
+			LogCaptureManagerIngestJobExecutor,
+			Error,
+			TEXT("[%s] Ingest failed, GEditor unavailble"), *ExecutorName
+		);
+
+		// Set the processing state to anything other than Pending, to prevent other executors picking up the same bad job.
+		InJob.SetProcessingState(FIngestJob::EProcessingState::Aborted);
+		InPromise.SetValue();
+		return nullptr;
+	}
+
+	UIngestCapability_Options* IngestOptions = NewObject<UIngestCapability_Options>();
+
+	// Naming tokens subsystem consults asset registry so need to run on the game thread
+	CallOnGameThread(
+		[&InJob, IngestOptions]()
+		{
+			UNamingTokensEditorSubsystem* NamingTokensSubsystem = GEditor->GetEditorSubsystem<UNamingTokensEditorSubsystem>();
+
+			FNamingTokenFilterArgs NamingTokenArgs;
+			if (const TObjectPtr<const UCaptureManagerNamingTokens> Tokens = GetDefault<UCaptureManagerSettings>()->GetNamingTokens())
+			{
+				NamingTokenArgs.NamespacesToInclude.Add(Tokens->GetNamespace());
+			}
+
+			const FIngestJob::FSettings Settings = InJob.GetSettings();
+
+			FStringFormatNamedArguments NamedArgs;
+			{
+				using namespace UE::CaptureManager;		
+				NamedArgs.Add(UCaptureManagerNamingTokens::ArchiveTokens[FString(GuidKey)].Name, InJob.GetTakeMetadata().UniqueId);
+				NamedArgs.Add(UCaptureManagerNamingTokens::ArchiveTokens[FString(DeviceKey)].Name, InJob.GetTakeMetadata().Device.UserId);
+				NamedArgs.Add(UCaptureManagerNamingTokens::ArchiveTokens[FString(SlateKey)].Name, InJob.GetTakeMetadata().Slate);
+				NamedArgs.Add(UCaptureManagerNamingTokens::ArchiveTokens[FString(TakeKey)].Name, InJob.GetTakeMetadata().TakeNumber);
+			}
+
+			FString WorkingDirectory = FString::Format(*Settings.LocalWorkFolder, NamedArgs);
+			FNamingTokenResultData WorkingDirectoryResult = NamingTokensSubsystem->EvaluateTokenString(WorkingDirectory, NamingTokenArgs);
+			IngestOptions->WorkingDirectory = WorkingDirectoryResult.EvaluatedText.ToString();
+
+			FString FileNamePrefix = FString::Format(*Settings.VideoSettings.FileNamePrefix, NamedArgs);
+			FNamingTokenResultData VideoPrefixResult = NamingTokensSubsystem->EvaluateTokenString(FileNamePrefix, NamingTokenArgs);
+			FIngestCapability_VideoOptions VideoOptions;
+			VideoOptions.FileNamePrefix = VideoPrefixResult.EvaluatedText.ToString();
+			VideoOptions.Format = Settings.VideoSettings.Format;
+			VideoOptions.PixelFormat = ConvertPixelFormat(Settings.VideoSettings.ImagePixelFormat);
+			IngestOptions->Video = VideoOptions;
+
+			FString AudioPrefix = FString::Format(*Settings.AudioSettings.FileNamePrefix, NamedArgs);
+			FNamingTokenResultData AudioPrefixResult = NamingTokensSubsystem->EvaluateTokenString(AudioPrefix, NamingTokenArgs);
+			FIngestCapability_AudioOptions AudioOptions;
+			AudioOptions.FileNamePrefix = AudioPrefixResult.EvaluatedText.ToString();
+			AudioOptions.Format = Settings.AudioSettings.Format;
+			IngestOptions->Audio = AudioOptions;
+		}
+	);
+
 	ULiveLinkDeviceSubsystem* DeviceSubsystem = GEngine->GetEngineSubsystem<ULiveLinkDeviceSubsystem>();
 
 	const TObjectPtr<ULiveLinkDevice>* DevicePtr = DeviceSubsystem->GetDeviceMap().Find(InJob.GetCaptureDeviceId());
@@ -198,20 +267,6 @@ UIngestCapability_ProcessHandle* FIngestJobExecutor::RunIngest(FIngestJob& InJob
 	ProcessHandle->ProcessFinishedReporter = MoveTemp(ProcessFinished);
 	ProcessHandle->ProcessProgressReporter = MoveTemp(ProgressReporter);
 
-	const FIngestJob::FSettings Settings = InJob.GetSettings();
-
-	UIngestCapability_Options* IngestOptions = NewObject<UIngestCapability_Options>();
-	IngestOptions->WorkingDirectory = Settings.LocalWorkFolder;
-	FIngestCapability_VideoOptions VideoOptions;
-	VideoOptions.FileNamePrefix = Settings.VideoSettings.FileNamePrefix;
-	VideoOptions.Format = Settings.VideoSettings.Format;
-	VideoOptions.PixelFormat = ConvertPixelFormat(Settings.VideoSettings.ImagePixelFormat);
-	IngestOptions->Video = VideoOptions;
-	FIngestCapability_AudioOptions AudioOptions;
-	AudioOptions.FileNamePrefix = Settings.AudioSettings.FileNamePrefix;
-	AudioOptions.Format = Settings.AudioSettings.Format;
-	IngestOptions->Audio = AudioOptions;
-
 	ProcessHandle->Run(IngestOptions);
 	return ProcessHandle;
 }
diff --git a/Engine/Plugins/VirtualProduction/CaptureManager/CaptureManagerApp/Source/CaptureManagerSettings/CaptureManagerSettings.Build.cs b/Engine/Plugins/VirtualProduction/CaptureManager/CaptureManagerApp/Source/CaptureManagerSettings/CaptureManagerSettings.Build.cs
--- a/Engine/Plugins/VirtualProduction/CaptureManager/CaptureManagerApp/Source/CaptureManagerSettings/CaptureManagerSettings.Build.cs
+++ b/Engine/Plugins/VirtualProduction/CaptureManager/CaptureManagerApp/Source/CaptureManagerSettings/CaptureManagerSettings.Build.cs
@@ -17,7 +17,13 @@ public class CaptureManagerSettings : ModuleRules
 
 		PrivateDependencyModuleNames.AddRange(new string[]
 		{
-			"Settings"
+			"Engine",
+			"NamingTokens",
+			"PropertyEditor",
+			"Slate",
+			"SlateCore",
+			"Settings",
+			"UnrealEd"
 		});
 	}
 }
diff --git a/Engine/Plugins/VirtualProduction/CaptureManager/CaptureManagerApp/Source/CaptureManagerSettings/Private/CaptureManagerSettingsModule.cpp b/Engine/Plugins/VirtualProduction/CaptureManager/CaptureManagerApp/Source/CaptureManagerSettings/Private/CaptureManagerSettingsModule.cpp
--- a/Engine/Plugins/VirtualProduction/CaptureManager/CaptureManagerApp/Source/CaptureManagerSettings/Private/CaptureManagerSettingsModule.cpp
+++ b/Engine/Plugins/VirtualProduction/CaptureManager/CaptureManagerApp/Source/CaptureManagerSettings/Private/CaptureManagerSettingsModule.cpp
@@ -3,8 +3,11 @@
 #include "CaptureManagerSettingsModule.h"
 
 #include "ISettingsModule.h"
-#include "Settings/CaptureManagerSettings.h"
 #include "Misc/CoreDelegates.h"
+#include "Modules/ModuleManager.h"
+#include "PropertyEditorModule.h"
+#include "Settings/CaptureManagerSettings.h"
+#include "Settings/CaptureManagerSettingsCustomization.h"
 
 #define LOCTEXT_NAMESPACE "CaptureManagerSettings"
 
@@ -12,6 +15,9 @@ void FCaptureManagerSettingsModule::StartupModule()
 {
 	FCoreDelegates::OnPostEngineInit.AddRaw(this, &FCaptureManagerSettingsModule::PostEngineInit);
 	FCoreDelegates::OnEnginePreExit.AddRaw(this, &FCaptureManagerSettingsModule::EnginePreExit);
+
+	FPropertyEditorModule& PropertyModule = FModuleManager::LoadModuleChecked<FPropertyEditorModule>("PropertyEditor");
+	PropertyModule.RegisterCustomClassLayout(UCaptureManagerSettings::StaticClass()->GetFName(), FOnGetDetailCustomizationInstance::CreateStatic(&FCaptureManagerSettingsCustomization::MakeInstance));
 }
 
 void FCaptureManagerSettingsModule::ShutdownModule()
diff --git a/Engine/Plugins/VirtualProduction/CaptureManager/CaptureManagerApp/Source/CaptureManagerSettings/Private/Settings/CaptureManagerSettings.cpp b/Engine/Plugins/VirtualProduction/CaptureManager/CaptureManagerApp/Source/CaptureManagerSettings/Private/Settings/CaptureManagerSettings.cpp
--- a/Engine/Plugins/VirtualProduction/CaptureManager/CaptureManagerApp/Source/CaptureManagerSettings/Private/Settings/CaptureManagerSettings.cpp
+++ b/Engine/Plugins/VirtualProduction/CaptureManager/CaptureManagerApp/Source/CaptureManagerSettings/Private/Settings/CaptureManagerSettings.cpp
@@ -1,17 +1,27 @@
 // Copyright Epic Games, Inc. All Rights Reserved.
 
 #include "Settings/CaptureManagerSettings.h"
+#include "Settings/CaptureManagerTemplateTokens.h"
 #include "Misc/Paths.h"
 
 namespace UE::CaptureManager::Private
 {
-constexpr FStringView DefaultConversionDirectory = TEXTVIEW("CaptureManagerConversion");
-constexpr FStringView DefaultDownloadDirectory = TEXTVIEW("CaptureManagerDownload");
+constexpr FStringView DefaultConversionDirectory = TEXTVIEW("CaptureManager/Conversion/{yyyy}{mm}{dd}");
+constexpr FStringView DefaultDownloadDirectory = TEXTVIEW("CaptureManager/Download/{yyyy}{mm}{dd}");
 }
 
 void UCaptureManagerSettings::Initialize()
 {
 	InitializeValuesIfNotSet();
+
+	NamingTokens = NewObject<UCaptureManagerNamingTokens>(const_cast<UCaptureManagerSettings*>(this),
+		UCaptureManagerNamingTokens::StaticClass());
+	NamingTokens->CreateDefaultTokens();
+}
+
+TObjectPtr<const UCaptureManagerNamingTokens> UCaptureManagerSettings::GetNamingTokens() const
+{	
+	return NamingTokens;
 }
 
 void UCaptureManagerSettings::PostEditChangeProperty(struct FPropertyChangedEvent& InPropertyChangedEvent)
diff --git a/Engine/Plugins/VirtualProduction/CaptureManager/CaptureManagerApp/Source/CaptureManagerSettings/Private/Settings/CaptureManagerSettingsCustomization.cpp b/Engine/Plugins/VirtualProduction/CaptureManager/CaptureManagerApp/Source/CaptureManagerSettings/Private/Settings/CaptureManagerSettingsCustomization.cpp
new file mode 100644
--- /dev/null
+++ b/Engine/Plugins/VirtualProduction/CaptureManager/CaptureManagerApp/Source/CaptureManagerSettings/Private/Settings/CaptureManagerSettingsCustomization.cpp
@@ -0,0 +1,69 @@
+// Copyright Epic Games, Inc. All Rights Reserved.
+
+#include "CaptureManagerSettingsCustomization.h"
+
+#include "Settings/CaptureManagerTemplateTokens.h"
+#include "Settings/CaptureManagerSettings.h"
+#include "DetailLayoutBuilder.h"
+#include "DetailWidgetRow.h"
+#include "Editor.h"
+#include "NamingTokensEditorSubsystem.h"
+#include "Widgets/Input/SEditableText.h"
+#include "Widgets/Text/STextBlock.h"
+
+#define LOCTEXT_NAMESPACE "CaptureManagerSettingsCustomization"
+
+TSharedRef<IDetailCustomization> FCaptureManagerSettingsCustomization::MakeInstance()
+{
+	return MakeShared<FCaptureManagerSettingsCustomization>();
+}
+
+void FCaptureManagerSettingsCustomization::CustomizeDetails(IDetailLayoutBuilder& InDetailBuilder)
+{
+	// Update current value when opening the settings page. Not safe to do under the settings object since there isn't
+	// an explicit callback when the settings page is opened, and using a method like PostInitProperties fires on the CDO
+	// early in the startup process.
+	
+	const TSharedRef<IPropertyHandle> AutomaticTokensHandle = InDetailBuilder.GetProperty(GET_MEMBER_NAME_CHECKED(UCaptureManagerSettings, AutomaticTokens));
+	IDetailPropertyRow* AutomaticTokensRow = InDetailBuilder.EditDefaultProperty(AutomaticTokensHandle);
+	check(AutomaticTokensRow);
+	
+	AutomaticTokensRow->CustomWidget()
+	.WholeRowContent()
+	[
+		SNew(SVerticalBox)
+		+ SVerticalBox::Slot()
+		.Padding(0.f, 4.f)
+		.AutoHeight()
+		[
+			SNew(STextBlock)
+			.Text(AutomaticTokensHandle->GetPropertyDisplayName())
+			.Font(IDetailLayoutBuilder::GetDetailFontBold())
+		]
+		+ SVerticalBox::Slot()
+		.AutoHeight()
+		[
+			SNew(SEditableText)
+			.IsReadOnly(true)
+			.Text_Raw(this, &FCaptureManagerSettingsCustomization::GetDisplayTokenText)
+			.Font(IDetailLayoutBuilder::GetDetailFont())
+		]
+	];
+}
+
+FText FCaptureManagerSettingsCustomization::GetDisplayTokenText() const
+{
+	FString FormattedTokensString = TEXT("None");
+	if (GEditor)
+	{
+		FNamingTokenFilterArgs Args;
+		if (const TObjectPtr<const UCaptureManagerNamingTokens> Tokens = GetDefault<UCaptureManagerSettings>()->GetNamingTokens())
+		{
+			Args.NamespacesToInclude.Add(Tokens->GetNamespace());
+		}
+		FormattedTokensString = GEditor->GetEditorSubsystem<UNamingTokensEditorSubsystem>()->GetFormattedTokensStringForDisplay(MoveTemp(Args));
+	}
+	return FText::FromString(FormattedTokensString);
+}
+
+#undef LOCTEXT_NAMESPACE
\ No newline at end of file
diff --git a/Engine/Plugins/VirtualProduction/CaptureManager/CaptureManagerApp/Source/CaptureManagerSettings/Private/Settings/CaptureManagerSettingsCustomization.h b/Engine/Plugins/VirtualProduction/CaptureManager/CaptureManagerApp/Source/CaptureManagerSettings/Private/Settings/CaptureManagerSettingsCustomization.h
new file mode 100644
--- /dev/null
+++ b/Engine/Plugins/VirtualProduction/CaptureManager/CaptureManagerApp/Source/CaptureManagerSettings/Private/Settings/CaptureManagerSettingsCustomization.h
@@ -0,0 +1,19 @@
+// Copyright Epic Games, Inc. All Rights Reserved.
+
+#pragma once
+
+#include "CoreMinimal.h"
+#include "IDetailCustomization.h"
+
+/**
+ * Customization for the capture manager settings.
+ */
+class FCaptureManagerSettingsCustomization : public IDetailCustomization
+{
+public:
+	static TSharedRef<IDetailCustomization> MakeInstance();
+	virtual void CustomizeDetails(IDetailLayoutBuilder& InDetailBuilder) override;
+	
+private:
+	FText GetDisplayTokenText() const;
+};
\ No newline at end of file
diff --git a/Engine/Plugins/VirtualProduction/CaptureManager/CaptureManagerApp/Source/CaptureManagerSettings/Private/Settings/CaptureManagerTemplateTokens.cpp b/Engine/Plugins/VirtualProduction/CaptureManager/CaptureManagerApp/Source/CaptureManagerSettings/Private/Settings/CaptureManagerTemplateTokens.cpp
new file mode 100644
--- /dev/null
+++ b/Engine/Plugins/VirtualProduction/CaptureManager/CaptureManagerApp/Source/CaptureManagerSettings/Private/Settings/CaptureManagerTemplateTokens.cpp
@@ -0,0 +1,46 @@
+// Copyright Epic Games, Inc. All Rights Reserved.
+
+#include "Settings/CaptureManagerTemplateTokens.h"
+
+#define LOCTEXT_NAMESPACE "CaptureManagerNamingTokens"
+
+TMap<FString, UE::CaptureManager::FArchiveToken> UCaptureManagerNamingTokens::ArchiveTokens = {
+	{FString(UE::CaptureManager::GuidKey), {TEXT("guid"), LOCTEXT("ArchiveGuid", "Archive Unique Id")}},
+	{FString(UE::CaptureManager::DeviceKey), {TEXT("device"), LOCTEXT("ArchiveDeviceId", "Archive Device User Id")}},
+	{FString(UE::CaptureManager::SlateKey), {TEXT("slate"), LOCTEXT("ArchiveSlate", "Archive Slate")}},
+	{FString(UE::CaptureManager::TakeKey), {TEXT("take"), LOCTEXT("ArchiveTake", "Archive Take")}}
+};
+
+UCaptureManagerNamingTokens::UCaptureManagerNamingTokens()
+{
+	Namespace = TEXT("cpman");
+}
+
+void UCaptureManagerNamingTokens::OnCreateDefaultTokens(TArray<FNamingTokenData>& Tokens)
+{
+	Super::OnCreateDefaultTokens(Tokens);
+
+	for (const TPair<FString, UE::CaptureManager::FArchiveToken>& Token : ArchiveTokens)
+	{
+		Tokens.Add({
+			Token.Value.Name,
+			Token.Value.Description,
+			FNamingTokenData::FTokenProcessorDelegateNative::CreateLambda([Token] {
+					return FText::FromString(Token.Value.Name);
+				})
+			});
+	}
+
+}
+
+void UCaptureManagerNamingTokens::OnPreEvaluate_Implementation(const FNamingTokensEvaluationData& InEvaluationData)
+{
+	Super::OnPreEvaluate_Implementation(InEvaluationData);
+}
+
+void UCaptureManagerNamingTokens::OnPostEvaluate_Implementation()
+{
+	Super::OnPostEvaluate_Implementation();
+}
+
+#undef LOCTEXT_NAMESPACE
diff --git a/Engine/Plugins/VirtualProduction/CaptureManager/CaptureManagerApp/Source/CaptureManagerSettings/Public/Settings/CaptureManagerSettings.h b/Engine/Plugins/VirtualProduction/CaptureManager/CaptureManagerApp/Source/CaptureManagerSettings/Public/Settings/CaptureManagerSettings.h
--- a/Engine/Plugins/VirtualProduction/CaptureManager/CaptureManagerApp/Source/CaptureManagerSettings/Public/Settings/CaptureManagerSettings.h
+++ b/Engine/Plugins/VirtualProduction/CaptureManager/CaptureManagerApp/Source/CaptureManagerSettings/Public/Settings/CaptureManagerSettings.h
@@ -6,6 +6,7 @@
 
 #include "CaptureManagerSettings.generated.h"
 
+class UCaptureManagerNamingTokens;
 
 UCLASS(config = Editor, defaultconfig)
 class CAPTUREMANAGERSETTINGS_API UCaptureManagerSettings : public UObject
@@ -15,6 +16,9 @@ public:
 
 	void Initialize();
 
+	/** Get the naming tokens for Capture Manager. */
+	TObjectPtr<const UCaptureManagerNamingTokens> GetNamingTokens() const;
+
 	/** Default location to where the converted data will be stored. It can be overriden when configuring the Ingest Job. */
 	UPROPERTY(config, EditAnywhere, Category = "Conversion")
 	FDirectoryPath DefaultConversionDirectory;
@@ -32,27 +36,38 @@ public:
 	bool bEnableThirdPartyEncoder = false;
 
 	/** Location to the third party encoder executable. */
-	UPROPERTY(config, EditAnywhere, Category = "Conversion", 
-			  meta = (EditCondition = "bEnableThirdPartyEncoder", EditConditionHides))
+	UPROPERTY(config, EditAnywhere, Category = "Conversion",
+		meta = (EditCondition = "bEnableThirdPartyEncoder", EditConditionHides))
 	FFilePath ThirdPartyEncoder;
 
 	/** Custom video command arguments to be used for executing the third party encoder.
 	 * NOTE: Leave empty to use the arguments generated from the Job settings.
 	 */
 	UPROPERTY(config, EditAnywhere, Category = "Conversion",
-			  meta = (EditCondition = "bEnableThirdPartyEncoder", EditConditionHides))
+		meta = (EditCondition = "bEnableThirdPartyEncoder", EditConditionHides))
 	FString CustomVideoCommandArguments;
 
-	/** Custom audio command arguments to be used for executing the third party encoder. 
+	/** Custom audio command arguments to be used for executing the third party encoder.
 	 * NOTE: Leave empty to use the arguments generated from the Job settings.
 	 */
 	UPROPERTY(config, EditAnywhere, Category = "Conversion",
-			  meta = (EditCondition = "bEnableThirdPartyEncoder", EditConditionHides))
+		meta = (EditCondition = "bEnableThirdPartyEncoder", EditConditionHides))
 	FString CustomAudioCommandArguments;
 
+	/** Placeholder for a list of the automatic tokens, set from the customization. */
+	UPROPERTY(VisibleAnywhere, Category = "Templates")
+	FText AutomaticTokens;
+
 private:
 
 	virtual void PostEditChangeProperty(struct FPropertyChangedEvent& InPropertyChangedEvent) override;
 
 	void InitializeValuesIfNotSet();
+
+	/**
+	* Naming tokens for Capture Manager, instantiated each load based on the naming tokens class.
+	* This isn't serialized to the config file, and exists here for singleton-like access.
+	*/
+	UPROPERTY(Instanced, Transient)
+	TObjectPtr<UCaptureManagerNamingTokens> NamingTokens;
 };
\ No newline at end of file
diff --git a/Engine/Plugins/VirtualProduction/CaptureManager/CaptureManagerApp/Source/CaptureManagerSettings/Public/Settings/CaptureManagerTemplateTokens.h b/Engine/Plugins/VirtualProduction/CaptureManager/CaptureManagerApp/Source/CaptureManagerSettings/Public/Settings/CaptureManagerTemplateTokens.h
new file mode 100644
--- /dev/null
+++ b/Engine/Plugins/VirtualProduction/CaptureManager/CaptureManagerApp/Source/CaptureManagerSettings/Public/Settings/CaptureManagerTemplateTokens.h
@@ -0,0 +1,40 @@
+// Copyright Epic Games, Inc. All Rights Reserved.
+
+#pragma once
+
+#include "NamingTokens.h"
+
+#include "CaptureManagerTemplateTokens.generated.h"
+
+namespace UE::CaptureManager
+{
+
+static constexpr FStringView GuidKey = TEXTVIEW("guid");
+static constexpr FStringView DeviceKey = TEXTVIEW("device");
+static constexpr FStringView SlateKey = TEXTVIEW("slate");
+static constexpr FStringView TakeKey = TEXTVIEW("take");
+
+struct FArchiveToken
+{
+	FString Name;
+	FText Description;
+};
+}
+
+UCLASS(NotBlueprintable)
+class CAPTUREMANAGERSETTINGS_API UCaptureManagerNamingTokens final : public UNamingTokens
+{
+	GENERATED_BODY()
+
+public:
+	UCaptureManagerNamingTokens();
+
+	static TMap<FString, UE::CaptureManager::FArchiveToken> ArchiveTokens;
+
+protected:
+	// ~Begin UNamingTokens
+	virtual void OnCreateDefaultTokens(TArray<FNamingTokenData>& Tokens) override;
+	virtual void OnPreEvaluate_Implementation(const FNamingTokensEvaluationData& InEvaluationData) override;
+	virtual void OnPostEvaluate_Implementation() override;
+	// ~End UNamingTokens
+};
diff --git a/Engine/Plugins/VirtualProduction/CaptureManager/CaptureManagerCore/Source/CaptureUtils/Private/Async/HelperFunctions.cpp b/Engine/Plugins/VirtualProduction/CaptureManager/CaptureManagerCore/Source/CaptureUtils/Private/Async/HelperFunctions.cpp
new file mode 100644
--- /dev/null
+++ b/Engine/Plugins/VirtualProduction/CaptureManager/CaptureManagerCore/Source/CaptureUtils/Private/Async/HelperFunctions.cpp
@@ -0,0 +1,33 @@
+// Copyright Epic Games, Inc. All Rights Reserved.
+
+#include "Async/HelperFunctions.h"
+
+#include "Async/Async.h"
+#include "Async/Future.h"
+#include "Async/TaskGraphInterfaces.h"
+#include "HAL/ThreadingBase.h"
+
+namespace UE::CaptureManager
+{
+
+void CallOnGameThread(TFunction<void()> InFunction)
+{
+	if (IsInGameThread())
+	{
+		InFunction();
+		return;
+	}
+
+	TPromise<void> Promise;
+	TFuture<void> Future = Promise.GetFuture();
+
+	AsyncTask(ENamedThreads::GameThread, [&Promise, &InFunction]
+		{
+			InFunction();
+			Promise.SetValue();
+		});
+
+	Future.Wait();
+}
+
+}
\ No newline at end of file
diff --git a/Engine/Plugins/VirtualProduction/CaptureManager/CaptureManagerCore/Source/CaptureUtils/Public/Async/HelperFunctions.h b/Engine/Plugins/VirtualProduction/CaptureManager/CaptureManagerCore/Source/CaptureUtils/Public/Async/HelperFunctions.h
new file mode 100644
--- /dev/null
+++ b/Engine/Plugins/VirtualProduction/CaptureManager/CaptureManagerCore/Source/CaptureUtils/Public/Async/HelperFunctions.h
@@ -0,0 +1,13 @@
+// Copyright Epic Games, Inc. All Rights Reserved.
+
+#pragma once
+
+#include "Templates/Function.h"
+
+namespace UE::CaptureManager
+{
+
+// Sends InFunction to be called on the GameThread and waits for the result
+void CAPTUREUTILS_API CallOnGameThread(TFunction<void()> InFunction);
+
+}
\ No newline at end of file
