Orientation should be added to the Conversion Settings

diff --git a/Engine/Plugins/VirtualProduction/CaptureManager/CaptureManagerApp/Source/CaptureDataConverter/Private/Nodes/CaptureConvertDepthData.cpp b/Engine/Plugins/VirtualProduction/CaptureManager/CaptureManagerApp/Source/CaptureDataConverter/Private/Nodes/CaptureConvertDepthData.cpp
--- a/Engine/Plugins/VirtualProduction/CaptureManager/CaptureManagerApp/Source/CaptureDataConverter/Private/Nodes/CaptureConvertDepthData.cpp
+++ b/Engine/Plugins/VirtualProduction/CaptureManager/CaptureManagerApp/Source/CaptureDataConverter/Private/Nodes/CaptureConvertDepthData.cpp
@@ -226,6 +226,8 @@ void FCaptureConvertDepthData::OnWrite(FWritingContext InContext)
 
 	IImageWriter* Writer = InContext.Writer;
 
+	DepthSample->Rotation = DepthParams.Rotation;
+
 	TOptional<FText> WriterResult = Writer->Append(DepthSample.Get());
 
 	if (WriterResult.IsSet())
diff --git a/Engine/Plugins/VirtualProduction/CaptureManager/CaptureManagerApp/Source/CaptureDataConverter/Private/Nodes/CaptureConvertVideoData.cpp b/Engine/Plugins/VirtualProduction/CaptureManager/CaptureManagerApp/Source/CaptureDataConverter/Private/Nodes/CaptureConvertVideoData.cpp
--- a/Engine/Plugins/VirtualProduction/CaptureManager/CaptureManagerApp/Source/CaptureDataConverter/Private/Nodes/CaptureConvertVideoData.cpp
+++ b/Engine/Plugins/VirtualProduction/CaptureManager/CaptureManagerApp/Source/CaptureDataConverter/Private/Nodes/CaptureConvertVideoData.cpp
@@ -218,6 +218,7 @@ void FCaptureConvertVideoData::OnWrite(FWritingContext InContext)
 
 	EMediaOrientation Orientation = ConvertOrientation(Video.Orientation.Get(FTakeMetadata::FVideo::EOrientation::Original));
 	VideoSample->Orientation = Orientation;
+	VideoSample->Rotation = VideoParams.Rotation;
 	VideoSample->DesiredFormat = VideoParams.OutputPixelFormat;
 
 	TOptional<FText> WriterResult = Writer->Append(VideoSample.Get());
@@ -250,7 +251,7 @@ EMediaOrientation FCaptureConvertVideoData::ConvertOrientation(FTakeMetadata::FV
 {
 	switch (InOrientation)
 	{
-	case FTakeMetadata::FVideo::EOrientation::CW90:
+		case FTakeMetadata::FVideo::EOrientation::CW90:
 			return EMediaOrientation::CW90;
 		case FTakeMetadata::FVideo::EOrientation::CW180:
 			return EMediaOrientation::CW180;
diff --git a/Engine/Plugins/VirtualProduction/CaptureManager/CaptureManagerApp/Source/CaptureDataConverter/Private/Nodes/ThirdPartyEncoder/CaptureConvertVideoDataThirdParty.cpp b/Engine/Plugins/VirtualProduction/CaptureManager/CaptureManagerApp/Source/CaptureDataConverter/Private/Nodes/ThirdPartyEncoder/CaptureConvertVideoDataThirdParty.cpp
--- a/Engine/Plugins/VirtualProduction/CaptureManager/CaptureManagerApp/Source/CaptureDataConverter/Private/Nodes/ThirdPartyEncoder/CaptureConvertVideoDataThirdParty.cpp
+++ b/Engine/Plugins/VirtualProduction/CaptureManager/CaptureManagerApp/Source/CaptureDataConverter/Private/Nodes/ThirdPartyEncoder/CaptureConvertVideoDataThirdParty.cpp
@@ -20,19 +20,19 @@ DEFINE_LOG_CATEGORY_STATIC(LogCaptureConvertVideoDataThirdParty, Log, All);
 namespace UE::CaptureManager::Private
 {
 
-static void ConvertOrientation(FTakeMetadata::FVideo::EOrientation InOrientation, TArray<FString>& OutFilterArray)
+static void ConvertOrientation(EMediaOrientation InOrientation, TArray<FString>& OutFilterArray)
 {
-	if (InOrientation == FTakeMetadata::FVideo::EOrientation::CW90)
+	if (InOrientation == EMediaOrientation::CW90)
 	{
-		OutFilterArray.Add(TEXT("transponse=cclock"));
+		OutFilterArray.Add(TEXT("transpose=clock"));
 	}
-	else if (InOrientation == FTakeMetadata::FVideo::EOrientation::CW180)
+	else if (InOrientation == EMediaOrientation::CW180)
 	{
-		OutFilterArray.Add(TEXT("transponse=clock,transponse=clock"));
+		OutFilterArray.Add(TEXT("transpose=clock,transpose=clock"));
 	}
-	else if (InOrientation == FTakeMetadata::FVideo::EOrientation::CW270)
+	else if (InOrientation == EMediaOrientation::CW270)
 	{
-		OutFilterArray.Add(TEXT("transpose=clock"));
+		OutFilterArray.Add(TEXT("transpose=cclock"));
 	}
 }
 
@@ -84,7 +84,7 @@ static void ConvertPixelFormat(EMediaTexturePixelFormat InPixelFormat, TArray<FS
 	}
 }
 
-static FString CreateConversionArguments(EMediaTexturePixelFormat InPixelFormat, FTakeMetadata::FVideo::EOrientation InOrientation)
+static FString CreateConversionArguments(EMediaTexturePixelFormat InPixelFormat, EMediaOrientation InOrientation)
 {
 	TArray<FString> FilterArray;
 
@@ -157,9 +157,7 @@ FCaptureConvertVideoDataThirdParty::FResult FCaptureConvertVideoDataThirdParty::
 	const FString ImageFileName = FString::Format(TEXT("{0}_%06d.{1}"), { VideoParams.ImageFileName, VideoParams.Format });
 	const FString ImageFilePath = FPaths::Combine(DestinationDirectory, ImageFileName);
 
-	FTakeMetadata::FVideo::EOrientation Orientation = Video.Orientation.Get(FTakeMetadata::FVideo::EOrientation::Original);
-
-	FString ConversionArguments = Private::CreateConversionArguments(VideoParams.OutputPixelFormat, Orientation);
+	FString ConversionArguments = Private::CreateConversionArguments(VideoParams.OutputPixelFormat, VideoParams.Rotation);
 
 	if (ThirdPartyEncoder.CommandArguments.IsEmpty())
 	{
diff --git a/Engine/Plugins/VirtualProduction/CaptureManager/CaptureManagerApp/Source/CaptureDataConverter/Public/CaptureDataConverterNodeParams.h b/Engine/Plugins/VirtualProduction/CaptureManager/CaptureManagerApp/Source/CaptureDataConverter/Public/CaptureDataConverterNodeParams.h
--- a/Engine/Plugins/VirtualProduction/CaptureManager/CaptureManagerApp/Source/CaptureDataConverter/Public/CaptureDataConverterNodeParams.h
+++ b/Engine/Plugins/VirtualProduction/CaptureManager/CaptureManagerApp/Source/CaptureDataConverter/Public/CaptureDataConverterNodeParams.h
@@ -19,6 +19,7 @@ struct FCaptureConvertVideoOutputParams
 	FString ImageFileName = TEXT("frame");
 	FString Format;
 	UE::CaptureManager::EMediaTexturePixelFormat OutputPixelFormat;
+	EMediaOrientation Rotation;
 };
 
 struct FCaptureConvertAudioOutputParams
@@ -31,6 +32,7 @@ struct FCaptureConvertDepthOutputParams
 {
 	FString ImageFileName = TEXT("depth");
 	bool bShouldCompressFiles = true;
+	EMediaOrientation Rotation;
 };
 
 struct FCaptureConvertCalibrationOutputParams
diff --git a/Engine/Plugins/VirtualProduction/CaptureManager/CaptureManagerApp/Source/CaptureManagerEditor/Private/IngestManagement/IngestJob.h b/Engine/Plugins/VirtualProduction/CaptureManager/CaptureManagerApp/Source/CaptureManagerEditor/Private/IngestManagement/IngestJob.h
--- a/Engine/Plugins/VirtualProduction/CaptureManager/CaptureManagerApp/Source/CaptureManagerEditor/Private/IngestManagement/IngestJob.h
+++ b/Engine/Plugins/VirtualProduction/CaptureManager/CaptureManagerApp/Source/CaptureManagerEditor/Private/IngestManagement/IngestJob.h
@@ -27,6 +27,7 @@ public:
 			FString Format;
 			FString FileNamePrefix;
 			::EImagePixelFormat ImagePixelFormat;
+			::EImageRotation ImageRotation;
 		};
 
 		struct FAudioSettings
diff --git a/Engine/Plugins/VirtualProduction/CaptureManager/CaptureManagerApp/Source/CaptureManagerEditor/Private/IngestManagement/IngestJobExecutor.cpp b/Engine/Plugins/VirtualProduction/CaptureManager/CaptureManagerApp/Source/CaptureManagerEditor/Private/IngestManagement/IngestJobExecutor.cpp
--- a/Engine/Plugins/VirtualProduction/CaptureManager/CaptureManagerApp/Source/CaptureManagerEditor/Private/IngestManagement/IngestJobExecutor.cpp
+++ b/Engine/Plugins/VirtualProduction/CaptureManager/CaptureManagerApp/Source/CaptureManagerEditor/Private/IngestManagement/IngestJobExecutor.cpp
@@ -24,15 +24,31 @@ DEFINE_LOG_CATEGORY_STATIC(LogCaptureManagerIngestJobExecutor, Log, All);
 namespace UE::CaptureManager
 {
 
-EIngestCapability_ImagePixelFormat ConvertPixelFormat(::EImagePixelFormat InImagePixelFormat)
+EIngestCapability_ImagePixelFormat ConvertImagePixelFormat(::EImagePixelFormat InImagePixelFormat)
 {
 	switch (InImagePixelFormat)
 	{
-	case ::EImagePixelFormat::U8_Mono:
-		return EIngestCapability_ImagePixelFormat::U8_Mono;
-	case ::EImagePixelFormat::U8_BGRA:
-	default:
-		return EIngestCapability_ImagePixelFormat::U8_BGRA;
+		case ::EImagePixelFormat::U8_Mono:
+			return EIngestCapability_ImagePixelFormat::U8_Mono;
+		case ::EImagePixelFormat::U8_BGRA:
+		default:
+			return EIngestCapability_ImagePixelFormat::U8_BGRA;
+	}
+}
+
+EIngestCapability_ImageRotation ConvertImageRotation(::EImageRotation InImageRotation)
+{
+	switch (InImageRotation)
+	{
+		case ::EImageRotation::CW_90:
+			return EIngestCapability_ImageRotation::CW_90;
+		case ::EImageRotation::CW_180:
+			return EIngestCapability_ImageRotation::CW_180;
+		case ::EImageRotation::CW_270:
+			return EIngestCapability_ImageRotation::CW_270;
+		case ::EImageRotation::None:
+		default:
+			return EIngestCapability_ImageRotation::None;
 	}
 }
 
@@ -182,7 +198,8 @@ UIngestCapability_ProcessHandle* FIngestJobExecutor::RunIngest(FIngestJob& InJob
 			FIngestCapability_VideoOptions VideoOptions;
 			VideoOptions.FileNamePrefix = VideoPrefixResult.EvaluatedText.ToString();
 			VideoOptions.Format = Settings.VideoSettings.Format;
-			VideoOptions.PixelFormat = ConvertPixelFormat(Settings.VideoSettings.ImagePixelFormat);
+			VideoOptions.PixelFormat = ConvertImagePixelFormat(Settings.VideoSettings.ImagePixelFormat);
+			VideoOptions.Rotation = ConvertImageRotation(Settings.VideoSettings.ImageRotation);
 			IngestOptions->Video = VideoOptions;
 
 			FString AudioPrefix = FString::Format(*Settings.AudioSettings.FileNamePrefix, NamedArgs);
diff --git a/Engine/Plugins/VirtualProduction/CaptureManager/CaptureManagerApp/Source/CaptureManagerEditor/Private/IngestManagement/UIngestJobSettings.cpp b/Engine/Plugins/VirtualProduction/CaptureManager/CaptureManagerApp/Source/CaptureManagerEditor/Private/IngestManagement/UIngestJobSettings.cpp
--- a/Engine/Plugins/VirtualProduction/CaptureManager/CaptureManagerApp/Source/CaptureManagerEditor/Private/IngestManagement/UIngestJobSettings.cpp
+++ b/Engine/Plugins/VirtualProduction/CaptureManager/CaptureManagerApp/Source/CaptureManagerEditor/Private/IngestManagement/UIngestJobSettings.cpp
@@ -21,6 +21,7 @@ UDefaultJobPreset::UDefaultJobPreset(const FObjectInitializer& ObjectInitializer
 	ImageFormat = EOutputImageFormat::JPEG;
 	ImageFileNamePrefix = TEXT("frame");
 	ImagePixelFormat = EImagePixelFormat::U8_BGRA;
+	ImageRotation = EImageRotation::None;
 	AudioFormat = EAudioFormat::WAV;
 	AudioFileNamePrefix = TEXT("audio");
 }
diff --git a/Engine/Plugins/VirtualProduction/CaptureManager/CaptureManagerApp/Source/CaptureManagerEditor/Private/IngestManagement/UIngestJobSettings.h b/Engine/Plugins/VirtualProduction/CaptureManager/CaptureManagerApp/Source/CaptureManagerEditor/Private/IngestManagement/UIngestJobSettings.h
--- a/Engine/Plugins/VirtualProduction/CaptureManager/CaptureManagerApp/Source/CaptureManagerEditor/Private/IngestManagement/UIngestJobSettings.h
+++ b/Engine/Plugins/VirtualProduction/CaptureManager/CaptureManagerApp/Source/CaptureManagerEditor/Private/IngestManagement/UIngestJobSettings.h
@@ -23,6 +23,15 @@ enum class EImagePixelFormat : uint8
 	U8_Mono				UMETA(DisplayName = "U8 Mono")
 };
 
+UENUM()
+enum class EImageRotation : uint8
+{
+	None = 0			UMETA(DisplayName = "None"),
+	CW_90				UMETA(DisplayName = "Clockwise 90°"),
+	CW_180				UMETA(DisplayName = "Clockwise 180°"),
+	CW_270				UMETA(DisplayName = "Clockwise 270°"),
+};
+
 UENUM()
 enum class EAudioFormat : uint8
 {
@@ -53,6 +62,9 @@ public:
 	UPROPERTY(Category = Video, EditAnywhere)
 	EImagePixelFormat ImagePixelFormat;
 
+	UPROPERTY(Category = Video, EditAnywhere)
+	EImageRotation ImageRotation;
+
 	UPROPERTY(Category = Audio, EditAnywhere)
 	EAudioFormat AudioFormat;
 
diff --git a/Engine/Plugins/VirtualProduction/CaptureManager/CaptureManagerApp/Source/CaptureManagerEditor/Private/UI/CaptureManagerPanelController.cpp b/Engine/Plugins/VirtualProduction/CaptureManager/CaptureManagerApp/Source/CaptureManagerEditor/Private/UI/CaptureManagerPanelController.cpp
--- a/Engine/Plugins/VirtualProduction/CaptureManager/CaptureManagerApp/Source/CaptureManagerEditor/Private/UI/CaptureManagerPanelController.cpp
+++ b/Engine/Plugins/VirtualProduction/CaptureManager/CaptureManagerApp/Source/CaptureManagerEditor/Private/UI/CaptureManagerPanelController.cpp
@@ -253,7 +253,8 @@ void FCaptureManagerPanelController::OnFinishedEditingJobProperties(const FPrope
 				{
 					.Format = StaticEnum<EOutputImageFormat>()->GetDisplayNameTextByValue(static_cast<std::underlying_type_t<EOutputImageFormat>>(Settings->ImageFormat)).ToString(),
 					.FileNamePrefix = Settings->ImageFileNamePrefix,
-					.ImagePixelFormat = Settings->ImagePixelFormat
+					.ImagePixelFormat = Settings->ImagePixelFormat,
+					.ImageRotation = Settings->ImageRotation
 				},
 				.AudioSettings =
 				{
diff --git a/Engine/Plugins/VirtualProduction/CaptureManager/CaptureManagerApp/Source/CaptureManagerMediaRW/Private/Writers/DepthImageWriter.cpp b/Engine/Plugins/VirtualProduction/CaptureManager/CaptureManagerApp/Source/CaptureManagerMediaRW/Private/Writers/DepthImageWriter.cpp
--- a/Engine/Plugins/VirtualProduction/CaptureManager/CaptureManagerApp/Source/CaptureManagerMediaRW/Private/Writers/DepthImageWriter.cpp
+++ b/Engine/Plugins/VirtualProduction/CaptureManager/CaptureManagerApp/Source/CaptureManagerMediaRW/Private/Writers/DepthImageWriter.cpp
@@ -79,9 +79,9 @@ TOptional<FText> FDepthExrImageWriter::Append(UE::CaptureManager::FMediaTextureS
 
 	const FString ExrFilePath = FPaths::Combine(Directory, FString::Printf(TEXT("%s_%06d.exr"), *FileName, FrameNumber));
 
-	FIntPoint OutputSize = GetOutputSize(InSample->Dimensions, InSample->Orientation);
+	FIntPoint OutputSize = GetOutputSize(InSample->Dimensions, InSample->Rotation);
 
-	TArray<float> RotatedData = Transform(MoveTemp(InSample->Buffer), InSample->Dimensions, InSample->Orientation);
+	TArray<float> RotatedData = Transform(MoveTemp(InSample->Buffer), InSample->Dimensions, InSample->Rotation);
 
 	if (!ImageWrapper->SetRaw(RotatedData.GetData(), RotatedData.Num() * sizeof(float), OutputSize.X, OutputSize.Y, ERGBFormat::GrayF, sizeof(float) * 8))
 	{
@@ -138,7 +138,7 @@ TArray<float> FDepthExrImageWriter::Transform(TArray<uint8> InDepthArray, FIntPo
 			{
 				for (int32 X = 0; X < OutputSize.X; ++X)
 				{
-					int32 Index = X * InDimensions.X + (OutputSize.Y - Y - 1);
+					int32 Index = (OutputSize.X - X - 1) * InDimensions.X + Y;
 					RotatedData[Y * OutputSize.X + X] = ConvertedData[Index] / TrueDepthResolutionPerCentimeter;
 				}
 			}
@@ -158,7 +158,7 @@ TArray<float> FDepthExrImageWriter::Transform(TArray<uint8> InDepthArray, FIntPo
 			{
 				for (int32 X = 0; X < OutputSize.X; ++X)
 				{
-					int32 Index = (OutputSize.X - X - 1) * InDimensions.X + Y;
+					int32 Index = X * InDimensions.X + (OutputSize.Y - Y - 1);
 					RotatedData[Y * OutputSize.X + X] = ConvertedData[Index] / TrueDepthResolutionPerCentimeter;
 				}
 			}
diff --git a/Engine/Plugins/VirtualProduction/CaptureManager/CaptureManagerApp/Source/CaptureManagerMediaRW/Private/Writers/WindowsImageWriter.cpp b/Engine/Plugins/VirtualProduction/CaptureManager/CaptureManagerApp/Source/CaptureManagerMediaRW/Private/Writers/WindowsImageWriter.cpp
--- a/Engine/Plugins/VirtualProduction/CaptureManager/CaptureManagerApp/Source/CaptureManagerMediaRW/Private/Writers/WindowsImageWriter.cpp
+++ b/Engine/Plugins/VirtualProduction/CaptureManager/CaptureManagerApp/Source/CaptureManagerMediaRW/Private/Writers/WindowsImageWriter.cpp
@@ -34,16 +34,16 @@ TUniquePtr<IImageWriter> FWindowsImageWriterFactory::CreateImageWriter()
 namespace UE::CaptureManager::Private
 {
 
-WICBitmapTransformOptions ConvertToWindowsTransformOptions(EMediaOrientation InOrientation)
+WICBitmapTransformOptions ConvertToWindowsTransformOptions(EMediaOrientation InRotation)
 {
-	switch (InOrientation)
+	switch (InRotation)
 	{
 		case EMediaOrientation::CW90:
-			return WICBitmapTransformRotate270;
+			return WICBitmapTransformRotate90;
 		case EMediaOrientation::CW180:
 			return WICBitmapTransformRotate180;
 		case EMediaOrientation::CW270:
-			return WICBitmapTransformRotate90;
+			return WICBitmapTransformRotate270;
 		case EMediaOrientation::Original:
 		default:
 			return WICBitmapTransformRotate0;
@@ -52,7 +52,7 @@ WICBitmapTransformOptions ConvertToWindowsTransformOptions(EMediaOrientation InO
 
 FIntPoint GetNewDimensions(const FMediaTextureSample* InNewSample)
 {
-	switch (InNewSample->Orientation)
+	switch (InNewSample->Rotation)
 	{
 		case EMediaOrientation::CW90:
 		case EMediaOrientation::CW270:
@@ -340,12 +340,12 @@ TOptional<FText> FWindowsImageWriter::Append(UE::CaptureManager::FMediaTextureSa
 	TComPtr<IWICBitmapFlipRotator> Rotator = nullptr;
 
 	GUID EncoderGuid = GetEncoderGuidBasedOnFormat(Format);
-	if (InSample->Orientation != EMediaOrientation::Original && EncoderGuid != GUID_ContainerFormatJpeg)
+	if (InSample->Rotation != EMediaOrientation::Original && EncoderGuid != GUID_ContainerFormatJpeg)
 	{
 		Result = WindowsImagingFactory->CreateBitmapFlipRotator(&Rotator);
 		WINIW_CHECK_AND_RETURN_ERROR_MESSAGE(Result, LOCTEXT("Open_FailedToCreateRotator", "Failed to create Rotator object"));
 
-		Result = Rotator->Initialize(Source, Private::ConvertToWindowsTransformOptions(InSample->Orientation));
+		Result = Rotator->Initialize(Source, Private::ConvertToWindowsTransformOptions(InSample->Rotation));
 		WINIW_CHECK_AND_RETURN_ERROR_MESSAGE(Result, LOCTEXT("Append_FailedToInitializeRotator", "Failed to initialize rotator from memory"));
 
 		Source = Rotator;
@@ -383,7 +383,7 @@ TOptional<FText> FWindowsImageWriter::Append(UE::CaptureManager::FMediaTextureSa
 		VARIANT Variant;
 		VariantInit(&Variant);
 		Variant.vt = VT_UI1;
-		Variant.bVal = Private::ConvertToWindowsTransformOptions(InSample->Orientation);
+		Variant.bVal = Private::ConvertToWindowsTransformOptions(InSample->Rotation);
 		constexpr ULONG NumberOfProperties = 1;
 		PropertyBag->Write(NumberOfProperties, &Options, &Variant);
 	}
diff --git a/Engine/Plugins/VirtualProduction/CaptureManager/CaptureManagerApp/Source/CaptureManagerMediaRW/Public/MediaSample.h b/Engine/Plugins/VirtualProduction/CaptureManager/CaptureManagerApp/Source/CaptureManagerMediaRW/Public/MediaSample.h
--- a/Engine/Plugins/VirtualProduction/CaptureManager/CaptureManagerApp/Source/CaptureManagerMediaRW/Public/MediaSample.h
+++ b/Engine/Plugins/VirtualProduction/CaptureManager/CaptureManagerApp/Source/CaptureManagerMediaRW/Public/MediaSample.h
@@ -61,6 +61,7 @@ struct FMediaTextureSample
 	EMediaTexturePixelFormat DesiredFormat;
 	FTimespan Time;
 	EMediaOrientation Orientation;
+	EMediaOrientation Rotation;
 };
 
 CAPTUREMANAGERMEDIARW_API uint32 GetNumberOfChannels(EMediaTexturePixelFormat InPixelFormat);
diff --git a/Engine/Plugins/VirtualProduction/CaptureManager/CaptureManagerApp/Source/IngestLiveLinkDevice/Private/BaseIngestLiveLinkDevice.cpp b/Engine/Plugins/VirtualProduction/CaptureManager/CaptureManagerApp/Source/IngestLiveLinkDevice/Private/BaseIngestLiveLinkDevice.cpp
--- a/Engine/Plugins/VirtualProduction/CaptureManager/CaptureManagerApp/Source/IngestLiveLinkDevice/Private/BaseIngestLiveLinkDevice.cpp
+++ b/Engine/Plugins/VirtualProduction/CaptureManager/CaptureManagerApp/Source/IngestLiveLinkDevice/Private/BaseIngestLiveLinkDevice.cpp
@@ -211,6 +211,7 @@ FCaptureDataConverterResult<void> UBaseIngestLiveLinkDevice::FImpl::RunConversio
 	VideoParams.Format = InIngestOptions->Video.Format;
 	VideoParams.ImageFileName = InIngestOptions->Video.FileNamePrefix;
 	VideoParams.OutputPixelFormat = Private::ConvertPixelFormat(InIngestOptions->Video.PixelFormat);
+	VideoParams.Rotation = Private::ConvertImageRotation(InIngestOptions->Video.Rotation);
 
 	InParams.VideoOutputParams = MoveTemp(VideoParams);
 
@@ -222,6 +223,7 @@ FCaptureDataConverterResult<void> UBaseIngestLiveLinkDevice::FImpl::RunConversio
 
 	FCaptureConvertDepthOutputParams DepthParams;
 	DepthParams.ImageFileName = TEXT("depth");
+	DepthParams.Rotation = Private::ConvertImageRotation(InIngestOptions->Video.Rotation); // Apply the same rotation as video to depth stream
 
 	InParams.DepthOutputParams = MoveTemp(DepthParams);
 
diff --git a/Engine/Plugins/VirtualProduction/CaptureManager/CaptureManagerApp/Source/IngestLiveLinkDevice/Private/Utils/IngestLiveLinkDeviceUtils.cpp b/Engine/Plugins/VirtualProduction/CaptureManager/CaptureManagerApp/Source/IngestLiveLinkDevice/Private/Utils/IngestLiveLinkDeviceUtils.cpp
--- a/Engine/Plugins/VirtualProduction/CaptureManager/CaptureManagerApp/Source/IngestLiveLinkDevice/Private/Utils/IngestLiveLinkDeviceUtils.cpp
+++ b/Engine/Plugins/VirtualProduction/CaptureManager/CaptureManagerApp/Source/IngestLiveLinkDevice/Private/Utils/IngestLiveLinkDeviceUtils.cpp
@@ -43,6 +43,22 @@ EMediaTexturePixelFormat ConvertPixelFormat(EIngestCapability_ImagePixelFormat I
 	}
 }
 
+EMediaOrientation ConvertImageRotation(EIngestCapability_ImageRotation InImageRotation)
+{
+	switch (InImageRotation)
+	{
+		case EIngestCapability_ImageRotation::CW_90:
+			return EMediaOrientation::CW90;
+		case EIngestCapability_ImageRotation::CW_180:
+			return EMediaOrientation::CW180;
+		case EIngestCapability_ImageRotation::CW_270:
+			return EMediaOrientation::CW270;
+		case EIngestCapability_ImageRotation::None:
+		default:
+			return EMediaOrientation::Original;
+	}
+}
+
 FString ErrorOriginToString(FTakeMetadataParserError::EOrigin InOrigin)
 {
 	switch (InOrigin)
diff --git a/Engine/Plugins/VirtualProduction/CaptureManager/CaptureManagerApp/Source/IngestLiveLinkDevice/Private/Utils/IngestLiveLinkDeviceUtils.h b/Engine/Plugins/VirtualProduction/CaptureManager/CaptureManagerApp/Source/IngestLiveLinkDevice/Private/Utils/IngestLiveLinkDeviceUtils.h
--- a/Engine/Plugins/VirtualProduction/CaptureManager/CaptureManagerApp/Source/IngestLiveLinkDevice/Private/Utils/IngestLiveLinkDeviceUtils.h
+++ b/Engine/Plugins/VirtualProduction/CaptureManager/CaptureManagerApp/Source/IngestLiveLinkDevice/Private/Utils/IngestLiveLinkDeviceUtils.h
@@ -12,6 +12,7 @@ namespace UE::CaptureManager::Private
 
 TSharedRef<FCommunicationManager> GetCommunicationManager();
 EMediaTexturePixelFormat ConvertPixelFormat(EIngestCapability_ImagePixelFormat InImagePixelFormat);
+EMediaOrientation ConvertImageRotation(EIngestCapability_ImageRotation InImageRotation);
 FString ErrorOriginToString(FTakeMetadataParserError::EOrigin InOrigin);
 
 }
\ No newline at end of file
diff --git a/Engine/Plugins/VirtualProduction/CaptureManager/CaptureManagerApp/Source/LiveLinkCapabilities/Public/Ingest/IngestCapability_Options.h b/Engine/Plugins/VirtualProduction/CaptureManager/CaptureManagerApp/Source/LiveLinkCapabilities/Public/Ingest/IngestCapability_Options.h
--- a/Engine/Plugins/VirtualProduction/CaptureManager/CaptureManagerApp/Source/LiveLinkCapabilities/Public/Ingest/IngestCapability_Options.h
+++ b/Engine/Plugins/VirtualProduction/CaptureManager/CaptureManagerApp/Source/LiveLinkCapabilities/Public/Ingest/IngestCapability_Options.h
@@ -24,6 +24,17 @@ enum class EIngestCapability_ImagePixelFormat
 	Default = U8_RGB	UMETA(DisplayName = "Default (U8_RGB)")
 };
 
+UENUM()
+enum class EIngestCapability_ImageRotation
+{
+	None = 0			UMETA(DisplayName = "None"),
+	CW_90				UMETA(DisplayName = "Clockwise 90°"),
+	CW_180				UMETA(DisplayName = "Clockwise 180°"),
+	CW_270				UMETA(DisplayName = "Clockwise 270°"),
+
+	Default = None		UMETA(DisplayName = "Default (None)")
+};
+
 USTRUCT(BlueprintType)
 struct LIVELINKCAPABILITIES_API FIngestCapability_VideoOptions
 {
@@ -37,6 +48,9 @@ struct LIVELINKCAPABILITIES_API FIngestCapability_VideoOptions
 
 	UPROPERTY(EditAnywhere, BlueprintReadWrite, Category = "Live Link Device|Ingest")
 	EIngestCapability_ImagePixelFormat PixelFormat;
+
+	UPROPERTY(EditAnywhere, BlueprintReadWrite, Category = "Live Link Device|Ingest")
+	EIngestCapability_ImageRotation Rotation;
 };
 
 USTRUCT(BlueprintType)
