Missing an icon of the Capture Manager tab within LiveLink Hub

diff --git a/Engine/Plugins/VirtualProduction/CaptureManager/CaptureManagerApp/Source/CaptureManagerEditor/CaptureManagerEditor.Build.cs b/Engine/Plugins/VirtualProduction/CaptureManager/CaptureManagerApp/Source/CaptureManagerEditor/CaptureManagerEditor.Build.cs
--- a/Engine/Plugins/VirtualProduction/CaptureManager/CaptureManagerApp/Source/CaptureManagerEditor/CaptureManagerEditor.Build.cs
+++ b/Engine/Plugins/VirtualProduction/CaptureManager/CaptureManagerApp/Source/CaptureManagerEditor/CaptureManagerEditor.Build.cs
@@ -38,7 +38,8 @@ public class CaptureManagerEditor : ModuleRules
 			"CaptureDataCore",
 			"CaptureManagerSettings",
 			"ImageCore",
-			"NamingTokens"
+			"NamingTokens",
+			"Projects"
 		});
 	}
 }
diff --git a/Engine/Plugins/VirtualProduction/CaptureManager/CaptureManagerApp/Source/CaptureManagerEditor/Private/CaptureManagerEditorModule.cpp b/Engine/Plugins/VirtualProduction/CaptureManager/CaptureManagerApp/Source/CaptureManagerEditor/Private/CaptureManagerEditorModule.cpp
--- a/Engine/Plugins/VirtualProduction/CaptureManager/CaptureManagerApp/Source/CaptureManagerEditor/Private/CaptureManagerEditorModule.cpp
+++ b/Engine/Plugins/VirtualProduction/CaptureManager/CaptureManagerApp/Source/CaptureManagerEditor/Private/CaptureManagerEditorModule.cpp
@@ -8,6 +8,8 @@
 #include "UI/TakeThumbnailRenderer.h"
 #include "UI/TakeVirtualAsset.h"
 
+#include "Styling/SlateStyleRegistry.h"
+#include "UI/Style/CaptureManagerTabStyle.h"
 
 
 IMPLEMENT_MODULE(FCaptureManagerEditorModule, CaptureManagerEditor)
@@ -17,6 +19,10 @@ void FCaptureManagerEditorModule::StartupModule()
 {
 	IModularFeatures::Get().RegisterModularFeature(ILiveLinkHubApplicationModeFactory::ModularFeatureName, this);
 	UThumbnailManager::Get().RegisterCustomRenderer(UTakeVirtualAsset::StaticClass(), UTakeThumbnailRenderer::StaticClass());
+
+	FSlateStyleRegistry::RegisterSlateStyle(FCaptureManagerTabStyle::Get());
+
+	FCaptureManagerTabStyle::ReloadTextures();
 }
 
 void FCaptureManagerEditorModule::ShutdownModule()
@@ -25,6 +31,8 @@ void FCaptureManagerEditorModule::ShutdownModule()
 	{
 		UThumbnailManager::Get().UnregisterCustomRenderer(UTakeThumbnailRenderer::StaticClass());
 	}
+
+	FSlateStyleRegistry::UnRegisterSlateStyle(FCaptureManagerTabStyle::Get());
 }
 
 /** Instantiate an application mode so LiveLinkHub can register it and display it in its Layout Selector. */
diff --git a/Engine/Plugins/VirtualProduction/CaptureManager/CaptureManagerApp/Source/CaptureManagerEditor/Private/UI/LiveLinkHubCaptureManagerMode.h b/Engine/Plugins/VirtualProduction/CaptureManager/CaptureManagerApp/Source/CaptureManagerEditor/Private/UI/LiveLinkHubCaptureManagerMode.h
--- a/Engine/Plugins/VirtualProduction/CaptureManager/CaptureManagerApp/Source/CaptureManagerEditor/Private/UI/LiveLinkHubCaptureManagerMode.h
+++ b/Engine/Plugins/VirtualProduction/CaptureManager/CaptureManagerApp/Source/CaptureManagerEditor/Private/UI/LiveLinkHubCaptureManagerMode.h
@@ -7,6 +7,8 @@
 #include "Framework/Docking/TabManager.h"
 #include "Templates/SharedPointer.h"
 
+#include "Style/CaptureManagerTabStyle.h"
+
 #include "Async/CaptureTimerManager.h"
 
 #include "IMessageContext.h"
@@ -31,7 +33,7 @@ public:
 	//~ Begin FLiveLinkHubApplicationMode interface
 	virtual FSlateIcon GetModeIcon() const override
 	{
-		return FSlateIcon(TEXT("MetaHumanIdentityStyle"), TEXT("ClassIcon.FootageCaptureData"), TEXT("ClassIcon.FootageCaptureData"));
+		return FSlateIcon(FCaptureManagerTabStyle::Get().GetStyleSetName(), TEXT("CaptureManager.Tabs.CaptureManager"));
 	}
 
 	/** Get widgets that should appear in the status bar. */
diff --git a/Engine/Plugins/VirtualProduction/CaptureManager/CaptureManagerApp/Source/CaptureManagerEditor/Private/UI/Style/CaptureManagerTabStyle.cpp b/Engine/Plugins/VirtualProduction/CaptureManager/CaptureManagerApp/Source/CaptureManagerEditor/Private/UI/Style/CaptureManagerTabStyle.cpp
new file mode 100644
--- /dev/null
+++ b/Engine/Plugins/VirtualProduction/CaptureManager/CaptureManagerApp/Source/CaptureManagerEditor/Private/UI/Style/CaptureManagerTabStyle.cpp
@@ -0,0 +1,39 @@
+// Copyright Epic Games, Inc. All Rights Reserved.
+
+#include "CaptureManagerTabStyle.h"
+#include "Framework/Application/SlateApplication.h"
+#include "Slate/SlateGameResources.h"
+#include "Interfaces/IPluginManager.h"
+#include "Styling/SlateStyleMacros.h"
+
+FName FCaptureManagerTabStyle::StyleName = TEXT("CaptureManagerTabStyle");
+
+FCaptureManagerTabStyle::FCaptureManagerTabStyle()
+	: FSlateStyleSet(FCaptureManagerTabStyle::StyleName)
+{
+	const FVector2D Icon16x16(16.0f, 16.0f);
+
+	SetContentRoot(IPluginManager::Get().FindPlugin(UE_PLUGIN_NAME)->GetContentDir());
+
+	Set("CaptureManager.Tabs.CaptureManager", new IMAGE_BRUSH_SVG("Icons/CaptureManager_16", Icon16x16));
+}
+
+const FName& FCaptureManagerTabStyle::GetStyleSetName() const
+{
+	return StyleName;
+}
+
+const FCaptureManagerTabStyle& FCaptureManagerTabStyle::Get()
+{
+	static FCaptureManagerTabStyle StyleInstance;
+
+	return StyleInstance;
+}
+
+void FCaptureManagerTabStyle::ReloadTextures()
+{
+	if (FSlateApplication::IsInitialized())
+	{
+		FSlateApplication::Get().GetRenderer()->ReloadTextureResources();
+	}
+}
diff --git a/Engine/Plugins/VirtualProduction/CaptureManager/CaptureManagerApp/Source/CaptureManagerEditor/Private/UI/Style/CaptureManagerTabStyle.h b/Engine/Plugins/VirtualProduction/CaptureManager/CaptureManagerApp/Source/CaptureManagerEditor/Private/UI/Style/CaptureManagerTabStyle.h
new file mode 100644
--- /dev/null
+++ b/Engine/Plugins/VirtualProduction/CaptureManager/CaptureManagerApp/Source/CaptureManagerEditor/Private/UI/Style/CaptureManagerTabStyle.h
@@ -0,0 +1,23 @@
+// Copyright Epic Games, Inc. All Rights Reserved.
+
+#pragma once
+
+#include "CoreMinimal.h"
+#include "Styling/SlateStyle.h"
+
+/** A Style set for Footage Retrieval window */
+class FCaptureManagerTabStyle : public FSlateStyleSet
+{
+public:
+
+	virtual const FName& GetStyleSetName() const override;
+	static const FCaptureManagerTabStyle& Get();
+
+	static void ReloadTextures();
+
+private:
+
+	FCaptureManagerTabStyle();
+
+	static FName StyleName;
+};
\ No newline at end of file
