[CaptureManager] Always check whether to clean up working directory on failed ingest

diff --git a/Engine/Plugins/VirtualProduction/CaptureManager/CaptureManagerApp/Source/CaptureManagerSettings/Public/Settings/CaptureManagerSettings.h b/Engine/Plugins/VirtualProduction/CaptureManager/CaptureManagerApp/Source/CaptureManagerSettings/Public/Settings/CaptureManagerSettings.h
--- a/Engine/Plugins/VirtualProduction/CaptureManager/CaptureManagerApp/Source/CaptureManagerSettings/Public/Settings/CaptureManagerSettings.h
+++ b/Engine/Plugins/VirtualProduction/CaptureManager/CaptureManagerApp/Source/CaptureManagerSettings/Public/Settings/CaptureManagerSettings.h
@@ -25,7 +25,7 @@ public:
 
 	/** Option to clean the converted data after the Ingest process. */
 	UPROPERTY(config, EditAnywhere, Category = "Conversion")
-	bool bShouldCleanConversionDirectory = true;
+	bool bShouldCleanWorkingDirectory = true;
 
 	/** Location where the take data downloaded from the device will be stored. */
 	UPROPERTY(config, EditAnywhere, Category = "Conversion")
diff --git a/Engine/Plugins/VirtualProduction/CaptureManager/CaptureManagerApp/Source/IngestLiveLinkDevice/Private/BaseIngestLiveLinkDevice.cpp b/Engine/Plugins/VirtualProduction/CaptureManager/CaptureManagerApp/Source/IngestLiveLinkDevice/Private/BaseIngestLiveLinkDevice.cpp
--- a/Engine/Plugins/VirtualProduction/CaptureManager/CaptureManagerApp/Source/IngestLiveLinkDevice/Private/BaseIngestLiveLinkDevice.cpp
+++ b/Engine/Plugins/VirtualProduction/CaptureManager/CaptureManagerApp/Source/IngestLiveLinkDevice/Private/BaseIngestLiveLinkDevice.cpp
@@ -34,6 +34,8 @@ private:
 	void UploadTakeStateHandler(const FGuid& InTakeUploadId, double InProgress);
 	void UploadTakeFinished(const FGuid& InTakeUploadId, const FString& InMessage, int32 InCode);
 
+	void TryCleanUpWorkingDirectory(const FString& InWorkingDirectory);
+
 	struct FContext
 	{
 		FString TakeName;
@@ -126,7 +128,7 @@ void UBaseIngestLiveLinkDevice::FImpl::IngestTake(int32 InTakeId, const UIngestC
 		}
 		else
 		{
-			IFileManager::Get().DeleteDirectory(*Params.TakeOutputDirectory, false, true);
+			TryCleanUpWorkingDirectory(*Params.TakeOutputDirectory);
 
 			FString Message = TEXT("Failed to upload take: Uploader not found");
 			Owner->ExecuteProcessFinishedReporter(InTakeId, MakeError(FIngestCapability_Error::UploaderNotFound, MoveTemp(Message)));
@@ -134,7 +136,7 @@ void UBaseIngestLiveLinkDevice::FImpl::IngestTake(int32 InTakeId, const UIngestC
 	}
 	else
 	{
-		IFileManager::Get().DeleteDirectory(*Params.TakeOutputDirectory, false, true);
+		TryCleanUpWorkingDirectory(*Params.TakeOutputDirectory);
 
 		FCaptureDataConverterError Error = Result.StealError();
 
@@ -185,7 +187,7 @@ void UBaseIngestLiveLinkDevice::FImpl::CancelIngest(int32 InTakeId)
 
 			const FString TakeDirectory = FPaths::Combine(Context.IngestOptions->WorkingDirectory,
 														  Context.TakeName + TEXT("_") + FoundContext->Key.ToString());
-			IFileManager::Get().DeleteDirectory(*TakeDirectory, false, true);
+			TryCleanUpWorkingDirectory(TakeDirectory);
 
 			FString Message = TEXT("Aborted by user");
 			Owner->ExecuteProcessFinishedReporter(InTakeId, MakeError(FIngestCapability_Error::AbortedByUser, MoveTemp(Message)));
@@ -268,13 +270,9 @@ void UBaseIngestLiveLinkDevice::FImpl::UploadTakeFinished(const FGuid& InTakeUpl
 	FScopeLock Lock(&Mutex);
 	if (const FContext* Context = ImportContexts.Find(InTakeUploadId))
 	{
-		const bool bShouldCleanConversionDirectory = GetDefault<UCaptureManagerSettings>()->bShouldCleanConversionDirectory;
-		if (bShouldCleanConversionDirectory)
-		{
-			const FString ConversionDirectory = FPaths::Combine(Context->IngestOptions->WorkingDirectory,
-																Context->TakeName + TEXT("_") + InTakeUploadId.ToString());
-			IFileManager::Get().DeleteDirectory(*ConversionDirectory, false, true);
-		}
+		const FString ConversionDirectory = FPaths::Combine(Context->IngestOptions->WorkingDirectory,
+															Context->TakeName + TEXT("_") + InTakeUploadId.ToString());
+		TryCleanUpWorkingDirectory(ConversionDirectory);
 
 		if (InMessage.IsEmpty())
 		{
@@ -291,6 +289,19 @@ void UBaseIngestLiveLinkDevice::FImpl::UploadTakeFinished(const FGuid& InTakeUpl
 	}
 }
 
+void UBaseIngestLiveLinkDevice::FImpl::TryCleanUpWorkingDirectory(const FString& InWorkingDirectory)
+{
+	const UCaptureManagerSettings* CaptureManagerSettings = GetDefault<UCaptureManagerSettings>();
+	check(CaptureManagerSettings);
+
+	const bool bShouldCleanWorkingDirectory = CaptureManagerSettings->bShouldCleanWorkingDirectory;
+	if (bShouldCleanWorkingDirectory)
+	{
+		IFileManager::Get().DeleteDirectory(*InWorkingDirectory, false, true);
+		UE_LOG(LogIngestLiveLinkDevice, Display, TEXT("Deleting directory: %s"), *InWorkingDirectory);
+	}
+}
+
 
 FGuid UBaseIngestLiveLinkDevice::GetDeviceGuid() const
 {
