[CaptureManager] Integrate naming tokens into the third party encoder commands

diff --git a/Engine/Plugins/VirtualProduction/CaptureManager/CaptureManagerApp/Source/CaptureDataConverter/CaptureDataConverter.Build.cs b/Engine/Plugins/VirtualProduction/CaptureManager/CaptureManagerApp/Source/CaptureDataConverter/CaptureDataConverter.Build.cs
--- a/Engine/Plugins/VirtualProduction/CaptureManager/CaptureManagerApp/Source/CaptureDataConverter/CaptureDataConverter.Build.cs
+++ b/Engine/Plugins/VirtualProduction/CaptureManager/CaptureManagerApp/Source/CaptureDataConverter/CaptureDataConverter.Build.cs
@@ -21,12 +21,15 @@ public class CaptureDataConverter : ModuleRules
 
 		PrivateDependencyModuleNames.AddRange(new string[]
 		{
+			"ContentBrowserData",
 			"CoreUObject",
 			"ImageWrapper",
 			"Json",
 			"DataIngestCore",
 			"Media",
-			"CaptureManagerSettings"
+			"CaptureManagerSettings",
+			"NamingTokens",
+			"UnrealEd"
 		});
 	}
 }
diff --git a/Engine/Plugins/VirtualProduction/CaptureManager/CaptureManagerApp/Source/CaptureDataConverter/Private/Nodes/ThirdPartyEncoder/CaptureConvertAudioDataThirdParty.cpp b/Engine/Plugins/VirtualProduction/CaptureManager/CaptureManagerApp/Source/CaptureDataConverter/Private/Nodes/ThirdPartyEncoder/CaptureConvertAudioDataThirdParty.cpp
--- a/Engine/Plugins/VirtualProduction/CaptureManager/CaptureManagerApp/Source/CaptureDataConverter/Private/Nodes/ThirdPartyEncoder/CaptureConvertAudioDataThirdParty.cpp
+++ b/Engine/Plugins/VirtualProduction/CaptureManager/CaptureManagerApp/Source/CaptureDataConverter/Private/Nodes/ThirdPartyEncoder/CaptureConvertAudioDataThirdParty.cpp
@@ -4,14 +4,22 @@
 
 #include "MediaSample.h"
 
-#include "Async/TaskProgress.h"
+#include "Async/HelperFunctions.h"
 #include "Async/StopToken.h"
+#include "Async/TaskProgress.h"
 
 #include "Nodes/CaptureCopyProgressReporter.h"
 #include "CaptureThirdPartyNodeUtils.h"
 
 #include "CaptureManagerMediaRWModule.h"
 
+#include "Editor.h"
+#include "GlobalNamingTokens.h"
+#include "NamingTokenData.h"
+#include "NamingTokensEditorSubsystem.h"
+#include "Settings/CaptureManagerSettings.h"
+#include "Settings/CaptureManagerTemplateTokens.h"
+
 #define LOCTEXT_NAMESPACE "CaptureConvertAudioDataTP"
 
 DEFINE_LOG_CATEGORY_STATIC(LogCaptureConvertAudioDataThirdParty, Log, All);
@@ -95,13 +103,34 @@ FCaptureConvertAudioDataThirdParty::FResult FCaptureConvertAudioDataThirdParty::
 		ThirdPartyEncoder.CommandArguments = UE::CaptureManager::AudioCommandArgumentTemplate;
 	}
 
-	FStringFormatNamedArguments FormatArguments;
-	FormatArguments.Add(TEXT("InputFileName"), WrapInQuotes(AudioFilePath));
-	FormatArguments.Add(TEXT("OutputFileName"), WrapInQuotes(AudioOutputFile));
+	FString CommandArgs = ThirdPartyEncoder.CommandArguments;
+
+	CallOnGameThread(
+		[&AudioFilePath, &AudioOutputFile, &CommandArgs]()
+		{
+			using namespace UE::CaptureManager;
 
-	const FString CommandArguments = FString::Format(*ThirdPartyEncoder.CommandArguments, FormatArguments);
+			UNamingTokensEditorSubsystem* NamingTokensSubsystem = GEditor->GetEditorSubsystem<UNamingTokensEditorSubsystem>();
+			check(NamingTokensSubsystem);
+
+			const UCaptureManagerSettings* Settings = GetDefault<UCaptureManagerSettings>();
+
+			FNamingTokenFilterArgs AudioEncoderTokenArgs;
+			const TObjectPtr<const UCaptureManagerAudioEncoderTokens> Tokens = Settings->GetAudioEncoderNamingTokens();
+			check(Tokens);
+			AudioEncoderTokenArgs.NamespacesToInclude.Add(Tokens->GetNamespace());
+
+			FStringFormatNamedArguments AudioEncoderFormatArgs;
+			AudioEncoderFormatArgs.Add(Tokens->GetToken(FString(AudioEncoderTokens::InputKey)).Name, WrapInQuotes(AudioFilePath));
+			AudioEncoderFormatArgs.Add(Tokens->GetToken(FString(AudioEncoderTokens::OutputKey)).Name, WrapInQuotes(AudioOutputFile));
+		
+			CommandArgs = FString::Format(*CommandArgs, AudioEncoderFormatArgs);
+			FNamingTokenResultData AudioEncoderCommandResult = NamingTokensSubsystem->EvaluateTokenString(CommandArgs, AudioEncoderTokenArgs);
+			CommandArgs = AudioEncoderCommandResult.EvaluatedText.ToString();
+		}
+	);
 
-	UE_LOG(LogCaptureConvertAudioDataThirdParty, Display, TEXT("Running the command: %s %s"), *ThirdPartyEncoder.Encoder, *CommandArguments);
+	UE_LOG(LogCaptureConvertAudioDataThirdParty, Display, TEXT("Running the command: %s %s"), *ThirdPartyEncoder.Encoder, *CommandArgs);
 
 	void* ReadPipe = nullptr;
 	void* WritePipe = nullptr;
@@ -111,7 +140,7 @@ FCaptureConvertAudioDataThirdParty::FResult FCaptureConvertAudioDataThirdParty::
 	constexpr bool bLaunchHidden = true;
 	constexpr bool bLaunchReallyHidden = true;
 	FProcHandle ProcHandle = 
-		FPlatformProcess::CreateProc(*ThirdPartyEncoder.Encoder, *CommandArguments, bLaunchDetached, bLaunchHidden, bLaunchReallyHidden, nullptr, 0, nullptr, WritePipe, nullptr);
+		FPlatformProcess::CreateProc(*ThirdPartyEncoder.Encoder, *CommandArgs, bLaunchDetached, bLaunchHidden, bLaunchReallyHidden, nullptr, 0, nullptr, WritePipe, nullptr);
 
 	ON_SCOPE_EXIT
 	{
@@ -128,7 +157,7 @@ FCaptureConvertAudioDataThirdParty::FResult FCaptureConvertAudioDataThirdParty::
 	{
 		FText Message = FText::Format(LOCTEXT("CaptureConvertAudioNodeTP_ProcessNotFound", "Failed to start the process {0} {1}"),
 									  FText::FromString(ThirdPartyEncoder.Encoder),
-									  FText::FromString(CommandArguments));
+									  FText::FromString(CommandArgs));
 		return MakeError(MoveTemp(Message));
 	}
 
@@ -160,7 +189,7 @@ FCaptureConvertAudioDataThirdParty::FResult FCaptureConvertAudioDataThirdParty::
 		if (!FullCommandOutput.IsEmpty())
 		{
 			UE_LOG(LogCaptureConvertAudioDataThirdParty, Error,
-				   TEXT("Failed to run the command: %s %s"), *ThirdPartyEncoder.Encoder, *CommandArguments);
+				   TEXT("Failed to run the command: %s %s"), *ThirdPartyEncoder.Encoder, *CommandArgs);
 
 			FString CommandOutputStr = FString::ConstructFromPtrSize(reinterpret_cast<const UTF8CHAR*>(FullCommandOutput.GetData()), FullCommandOutput.Num());
 			UE_LOG(LogCaptureConvertAudioDataThirdParty, Display,
diff --git a/Engine/Plugins/VirtualProduction/CaptureManager/CaptureManagerApp/Source/CaptureDataConverter/Private/Nodes/ThirdPartyEncoder/CaptureConvertVideoDataThirdParty.cpp b/Engine/Plugins/VirtualProduction/CaptureManager/CaptureManagerApp/Source/CaptureDataConverter/Private/Nodes/ThirdPartyEncoder/CaptureConvertVideoDataThirdParty.cpp
--- a/Engine/Plugins/VirtualProduction/CaptureManager/CaptureManagerApp/Source/CaptureDataConverter/Private/Nodes/ThirdPartyEncoder/CaptureConvertVideoDataThirdParty.cpp
+++ b/Engine/Plugins/VirtualProduction/CaptureManager/CaptureManagerApp/Source/CaptureDataConverter/Private/Nodes/ThirdPartyEncoder/CaptureConvertVideoDataThirdParty.cpp
@@ -4,13 +4,21 @@
 
 #include "MediaSample.h"
 
-#include "Async/TaskProgress.h"
+#include "Async/HelperFunctions.h"
 #include "Async/StopToken.h"
+#include "Async/TaskProgress.h"
 
 #include "CaptureManagerMediaRWModule.h"
 #include "Nodes/CaptureCopyProgressReporter.h"
 #include "CaptureThirdPartyNodeUtils.h"
 
+#include "Editor.h"
+#include "GlobalNamingTokens.h"
+#include "NamingTokenData.h"
+#include "NamingTokensEditorSubsystem.h"
+#include "Settings/CaptureManagerSettings.h"
+#include "Settings/CaptureManagerTemplateTokens.h"
+
 #include "IImageWrapperModule.h"
 
 #define LOCTEXT_NAMESPACE "CaptureConvertVideoDataTP"
@@ -113,7 +121,7 @@ FCaptureConvertVideoDataThirdParty::FCaptureConvertVideoDataThirdParty(FCaptureT
 	, Params(InParams)
 	, VideoParams(InVideoParams)
 {
-	checkf(!VideoParams.Format.IsEmpty(), TEXT("Format MUST be specified"));
+	checkf(!VideoParams.Format.IsEmpty(), TEXT("Video output format MUST be specified"));
 }
 
 FCaptureConvertVideoDataThirdParty::FResult FCaptureConvertVideoDataThirdParty::Run()
@@ -164,14 +172,35 @@ FCaptureConvertVideoDataThirdParty::FResult FCaptureConvertVideoDataThirdParty::
 		ThirdPartyEncoder.CommandArguments = UE::CaptureManager::VideoCommandArgumentTemplate;
 	}
 
-	FStringFormatNamedArguments FormatArguments;
-	FormatArguments.Add(TEXT("InputFileName"), WrapInQuotes(VideoFilePath));
-	FormatArguments.Add(TEXT("OutputFileName"), WrapInQuotes(ImageFilePath));
-	FormatArguments.Add(TEXT("ConversionFilters"), ConversionArguments);
+	FString CommandArgs = ThirdPartyEncoder.CommandArguments;
 
-	FString CommandArguments = FString::Format(*ThirdPartyEncoder.CommandArguments, FormatArguments);
+	CallOnGameThread(
+		[&ConversionArguments, &VideoFilePath, &ImageFilePath, &CommandArgs]()
+		{
+			using namespace UE::CaptureManager;
+
+			UNamingTokensEditorSubsystem* NamingTokensSubsystem = GEditor->GetEditorSubsystem<UNamingTokensEditorSubsystem>();
+			check(NamingTokensSubsystem);
+
+			const UCaptureManagerSettings* Settings = GetDefault<UCaptureManagerSettings>();
+
+			FNamingTokenFilterArgs VideoEncoderTokenArgs;
+			const TObjectPtr<const UCaptureManagerVideoEncoderTokens> Tokens = Settings->GetVideoEncoderNamingTokens();
+			check(Tokens);
+			VideoEncoderTokenArgs.NamespacesToInclude.Add(Tokens->GetNamespace());
+
+			FStringFormatNamedArguments VideoEncoderFormatArgs;
+			VideoEncoderFormatArgs.Add(Tokens->GetToken(FString(VideoEncoderTokens::InputKey)).Name, WrapInQuotes(VideoFilePath));
+			VideoEncoderFormatArgs.Add(Tokens->GetToken(FString(VideoEncoderTokens::OutputKey)).Name, WrapInQuotes(ImageFilePath));
+			VideoEncoderFormatArgs.Add(Tokens->GetToken(FString(VideoEncoderTokens::ParamsKey)).Name, ConversionArguments);
+			
+			CommandArgs = FString::Format(*CommandArgs, VideoEncoderFormatArgs);
+			FNamingTokenResultData VideoEncoderCommandResult = NamingTokensSubsystem->EvaluateTokenString(CommandArgs, VideoEncoderTokenArgs);
+			CommandArgs = VideoEncoderCommandResult.EvaluatedText.ToString();
+		}
+	);
 	
-	UE_LOG(LogCaptureConvertVideoDataThirdParty, Display, TEXT("Running the command: %s %s"), *ThirdPartyEncoder.Encoder, *CommandArguments);
+	UE_LOG(LogCaptureConvertVideoDataThirdParty, Display, TEXT("Running the command: %s %s"), *ThirdPartyEncoder.Encoder, *CommandArgs);
 
 	void* ReadPipe = nullptr;
 	void* WritePipe = nullptr;
@@ -181,7 +210,7 @@ FCaptureConvertVideoDataThirdParty::FResult FCaptureConvertVideoDataThirdParty::
 	constexpr bool bLaunchHidden = true;
 	constexpr bool bLaunchReallyHidden = true;
 	FProcHandle ProcHandle = 
-		FPlatformProcess::CreateProc(*ThirdPartyEncoder.Encoder, *CommandArguments, bLaunchDetached, bLaunchHidden, bLaunchReallyHidden, nullptr, 0, nullptr, WritePipe, nullptr);
+		FPlatformProcess::CreateProc(*ThirdPartyEncoder.Encoder, *CommandArgs, bLaunchDetached, bLaunchHidden, bLaunchReallyHidden, nullptr, 0, nullptr, WritePipe, nullptr);
 
 	ON_SCOPE_EXIT
 	{
@@ -198,7 +227,7 @@ FCaptureConvertVideoDataThirdParty::FResult FCaptureConvertVideoDataThirdParty::
 	{
 		FText Message = FText::Format(LOCTEXT("CaptureConvertVideoNodeTP_ProcessNotFound", "Failed to start the process {0} {1}"),
 									  FText::FromString(ThirdPartyEncoder.Encoder),
-									  FText::FromString(CommandArguments));
+									  FText::FromString(CommandArgs));
 		return MakeError(MoveTemp(Message));
 	}
 
@@ -231,7 +260,7 @@ FCaptureConvertVideoDataThirdParty::FResult FCaptureConvertVideoDataThirdParty::
 		if (!FullCommandOutput.IsEmpty())
 		{
 			UE_LOG(LogCaptureConvertVideoDataThirdParty, Error,
-				   TEXT("Failed to run the command: %s %s"), *ThirdPartyEncoder.Encoder, *CommandArguments);
+				   TEXT("Failed to run the command: %s %s"), *ThirdPartyEncoder.Encoder, *CommandArgs);
 
 			FString CommandOutputStr = FString::ConstructFromPtrSize(reinterpret_cast<const UTF8CHAR*>(FullCommandOutput.GetData()), FullCommandOutput.Num());
 			UE_LOG(LogCaptureConvertVideoDataThirdParty, Display,
diff --git a/Engine/Plugins/VirtualProduction/CaptureManager/CaptureManagerApp/Source/CaptureDataConverter/Private/Nodes/ThirdPartyEncoder/CaptureThirdPartyNodeParams.h b/Engine/Plugins/VirtualProduction/CaptureManager/CaptureManagerApp/Source/CaptureDataConverter/Private/Nodes/ThirdPartyEncoder/CaptureThirdPartyNodeParams.h
--- a/Engine/Plugins/VirtualProduction/CaptureManager/CaptureManagerApp/Source/CaptureDataConverter/Private/Nodes/ThirdPartyEncoder/CaptureThirdPartyNodeParams.h
+++ b/Engine/Plugins/VirtualProduction/CaptureManager/CaptureManagerApp/Source/CaptureDataConverter/Private/Nodes/ThirdPartyEncoder/CaptureThirdPartyNodeParams.h
@@ -7,8 +7,8 @@
 namespace UE::CaptureManager
 {
 
-static const FString VideoCommandArgumentTemplate = TEXT("-noautorotate -i {InputFileName} -an {ConversionFilters} -q:v 1 -qmax 1 -qmin 1 -start_number 0 {OutputFileName}");
-static const FString AudioCommandArgumentTemplate = TEXT("-i {InputFileName} -vn {OutputFileName}");
+static const FString VideoCommandArgumentTemplate = TEXT("-noautorotate -i {input} -an {params} -q:v 1 -qmax 1 -qmin 1 -start_number 0 {output}");
+static const FString AudioCommandArgumentTemplate = TEXT("-i {input} -vn {output}");
 
 }
 
diff --git a/Engine/Plugins/VirtualProduction/CaptureManager/CaptureManagerApp/Source/CaptureManagerEditor/Private/IngestManagement/IngestJobExecutor.cpp b/Engine/Plugins/VirtualProduction/CaptureManager/CaptureManagerApp/Source/CaptureManagerEditor/Private/IngestManagement/IngestJobExecutor.cpp
--- a/Engine/Plugins/VirtualProduction/CaptureManager/CaptureManagerApp/Source/CaptureManagerEditor/Private/IngestManagement/IngestJobExecutor.cpp
+++ b/Engine/Plugins/VirtualProduction/CaptureManager/CaptureManagerApp/Source/CaptureManagerEditor/Private/IngestManagement/IngestJobExecutor.cpp
@@ -173,29 +173,29 @@ UIngestCapability_ProcessHandle* FIngestJobExecutor::RunIngest(FIngestJob& InJob
 			UNamingTokensEditorSubsystem* NamingTokensSubsystem = GEditor->GetEditorSubsystem<UNamingTokensEditorSubsystem>();
 			check(NamingTokensSubsystem);
 
-			FNamingTokenFilterArgs NamingTokenArgs;
-			if (const TObjectPtr<const UCaptureManagerGeneralTokens> Tokens = GetDefault<UCaptureManagerSettings>()->GetGeneralTokens())
-			{
-				NamingTokenArgs.NamespacesToInclude.Add(Tokens->GetNamespace());
-			}
+			FNamingTokenFilterArgs GeneralNamingTokenArgs;
+			const TObjectPtr<const UCaptureManagerGeneralTokens> Tokens = GetDefault<UCaptureManagerSettings>()->GetGeneralNamingTokens();
+			check(Tokens);
+
+			GeneralNamingTokenArgs.NamespacesToInclude.Add(Tokens->GetNamespace());
 
 			const FIngestJob::FSettings Settings = InJob.GetSettings();
 
-			FStringFormatNamedArguments NamedArgs;
+			FStringFormatNamedArguments GeneralArgs;
 			{
-				using namespace UE::CaptureManager;		
-				NamedArgs.Add(UCaptureManagerGeneralTokens::GeneralTokens[FString(UE::CaptureManager::IdKey)].Name, InJob.GetTakeMetadata().UniqueId);
-				NamedArgs.Add(UCaptureManagerGeneralTokens::GeneralTokens[FString(UE::CaptureManager::DeviceKey)].Name, InJob.GetTakeMetadata().Device.UserId);
-				NamedArgs.Add(UCaptureManagerGeneralTokens::GeneralTokens[FString(UE::CaptureManager::SlateKey)].Name, InJob.GetTakeMetadata().Slate);
-				NamedArgs.Add(UCaptureManagerGeneralTokens::GeneralTokens[FString(UE::CaptureManager::TakeKey)].Name, InJob.GetTakeMetadata().TakeNumber);
+				using namespace UE::CaptureManager;
+				GeneralArgs.Add(Tokens->GetToken(FString(GeneralTokens::IdKey)).Name, InJob.GetTakeMetadata().UniqueId);
+				GeneralArgs.Add(Tokens->GetToken(FString(GeneralTokens::DeviceKey)).Name, InJob.GetTakeMetadata().Device.UserId);
+				GeneralArgs.Add(Tokens->GetToken(FString(GeneralTokens::SlateKey)).Name, InJob.GetTakeMetadata().Slate);
+				GeneralArgs.Add(Tokens->GetToken(FString(GeneralTokens::TakeKey)).Name, InJob.GetTakeMetadata().TakeNumber);
 			}
 
-			FString WorkingDirectory = FString::Format(*Settings.LocalWorkFolder, NamedArgs);
-			FNamingTokenResultData WorkingDirectoryResult = NamingTokensSubsystem->EvaluateTokenString(WorkingDirectory, NamingTokenArgs);
+			FString WorkingDirectory = FString::Format(*Settings.LocalWorkFolder, GeneralArgs);
+			FNamingTokenResultData WorkingDirectoryResult = NamingTokensSubsystem->EvaluateTokenString(WorkingDirectory, GeneralNamingTokenArgs);
 			IngestOptions->WorkingDirectory = WorkingDirectoryResult.EvaluatedText.ToString();
 
-			FString FileNamePrefix = FString::Format(*Settings.VideoSettings.FileNamePrefix, NamedArgs);
-			FNamingTokenResultData VideoPrefixResult = NamingTokensSubsystem->EvaluateTokenString(FileNamePrefix, NamingTokenArgs);
+			FString FileNamePrefix = FString::Format(*Settings.VideoSettings.FileNamePrefix, GeneralArgs);
+			FNamingTokenResultData VideoPrefixResult = NamingTokensSubsystem->EvaluateTokenString(FileNamePrefix, GeneralNamingTokenArgs);
 			FIngestCapability_VideoOptions VideoOptions;
 			VideoOptions.FileNamePrefix = VideoPrefixResult.EvaluatedText.ToString();
 			VideoOptions.Format = Settings.VideoSettings.Format;
@@ -203,8 +203,8 @@ UIngestCapability_ProcessHandle* FIngestJobExecutor::RunIngest(FIngestJob& InJob
 			VideoOptions.Rotation = ConvertImageRotation(Settings.VideoSettings.ImageRotation);
 			IngestOptions->Video = VideoOptions;
 
-			FString AudioPrefix = FString::Format(*Settings.AudioSettings.FileNamePrefix, NamedArgs);
-			FNamingTokenResultData AudioPrefixResult = NamingTokensSubsystem->EvaluateTokenString(AudioPrefix, NamingTokenArgs);
+			FString AudioPrefix = FString::Format(*Settings.AudioSettings.FileNamePrefix, GeneralArgs);
+			FNamingTokenResultData AudioPrefixResult = NamingTokensSubsystem->EvaluateTokenString(AudioPrefix, GeneralNamingTokenArgs);
 			FIngestCapability_AudioOptions AudioOptions;
 			AudioOptions.FileNamePrefix = AudioPrefixResult.EvaluatedText.ToString();
 			AudioOptions.Format = Settings.AudioSettings.Format;
diff --git a/Engine/Plugins/VirtualProduction/CaptureManager/CaptureManagerApp/Source/CaptureManagerSettings/Private/Settings/CaptureManagerSettings.cpp b/Engine/Plugins/VirtualProduction/CaptureManager/CaptureManagerApp/Source/CaptureManagerSettings/Private/Settings/CaptureManagerSettings.cpp
--- a/Engine/Plugins/VirtualProduction/CaptureManager/CaptureManagerApp/Source/CaptureManagerSettings/Private/Settings/CaptureManagerSettings.cpp
+++ b/Engine/Plugins/VirtualProduction/CaptureManager/CaptureManagerApp/Source/CaptureManagerSettings/Private/Settings/CaptureManagerSettings.cpp
@@ -14,13 +14,29 @@ void UCaptureManagerSettings::Initialize()
 {
 	InitializeValuesIfNotSet();
 
-	GeneralTokens = NewObject<UCaptureManagerGeneralTokens>(this, UCaptureManagerGeneralTokens::StaticClass());
-	GeneralTokens->CreateDefaultTokens();
+	GeneralNamingTokens = NewObject<UCaptureManagerGeneralTokens>(this, UCaptureManagerGeneralTokens::StaticClass());
+	GeneralNamingTokens->CreateDefaultTokens();
+
+	VideoEncoderNamingTokens = NewObject<UCaptureManagerVideoEncoderTokens>(this, UCaptureManagerVideoEncoderTokens::StaticClass());
+	VideoEncoderNamingTokens->CreateDefaultTokens();
+
+	AudioEncoderNamingTokens = NewObject<UCaptureManagerAudioEncoderTokens>(this, UCaptureManagerAudioEncoderTokens::StaticClass());
+	AudioEncoderNamingTokens->CreateDefaultTokens();
 }
 
-TObjectPtr<const UCaptureManagerGeneralTokens> UCaptureManagerSettings::GetGeneralTokens() const
+TObjectPtr<UCaptureManagerGeneralTokens> UCaptureManagerSettings::GetGeneralNamingTokens() const
 {	
-	return GeneralTokens;
+	return GeneralNamingTokens;
+}
+
+TObjectPtr<UCaptureManagerVideoEncoderTokens> UCaptureManagerSettings::GetVideoEncoderNamingTokens() const
+{
+	return VideoEncoderNamingTokens;
+}
+
+TObjectPtr<UCaptureManagerAudioEncoderTokens> UCaptureManagerSettings::GetAudioEncoderNamingTokens() const
+{
+	return AudioEncoderNamingTokens;
 }
 
 void UCaptureManagerSettings::PostEditChangeProperty(struct FPropertyChangedEvent& InPropertyChangedEvent)
diff --git a/Engine/Plugins/VirtualProduction/CaptureManager/CaptureManagerApp/Source/CaptureManagerSettings/Private/Settings/CaptureManagerSettingsCustomization.cpp b/Engine/Plugins/VirtualProduction/CaptureManager/CaptureManagerApp/Source/CaptureManagerSettings/Private/Settings/CaptureManagerSettingsCustomization.cpp
--- a/Engine/Plugins/VirtualProduction/CaptureManager/CaptureManagerApp/Source/CaptureManagerSettings/Private/Settings/CaptureManagerSettingsCustomization.cpp
+++ b/Engine/Plugins/VirtualProduction/CaptureManager/CaptureManagerApp/Source/CaptureManagerSettings/Private/Settings/CaptureManagerSettingsCustomization.cpp
@@ -23,45 +23,80 @@ void FCaptureManagerSettingsCustomization::CustomizeDetails(IDetailLayoutBuilder
 	// Update current value when opening the settings page. Not safe to do under the settings object since there isn't
 	// an explicit callback when the settings page is opened, and using a method like PostInitProperties fires on the CDO
 	// early in the startup process.
+
+	auto BuildSlate = [this](const TSharedRef<IPropertyHandle> InHandle, IDetailPropertyRow* InPropertyRow, const FNamingTokenFilterArgs& InTokenArgs) {
+		InPropertyRow->CustomWidget()
+			.WholeRowContent()
+			[
+				SNew(SVerticalBox)
+					+ SVerticalBox::Slot()
+					.Padding(0.f, 4.f)
+					.AutoHeight()
+					[
+						SNew(STextBlock)
+							.Text(InHandle->GetPropertyDisplayName())
+							.Font(IDetailLayoutBuilder::GetDetailFontBold())
+					]
+					+ SVerticalBox::Slot()
+					.AutoHeight()
+					[
+						SNew(SEditableText)
+							.IsReadOnly(true)
+							.Text_Raw(this, &FCaptureManagerSettingsCustomization::GetDisplayTokenText, InTokenArgs)
+							.Font(IDetailLayoutBuilder::GetDetailFont())
+					]
+			];
+		};
+
+	FNamingTokenFilterArgs BaseArgs =
+	{
+		.NamespacesToInclude = {},
+		.bIncludeGlobal = false
+	};
+
+	// General tokens
+	const TSharedRef<IPropertyHandle> GeneralTokensHandle = InDetailBuilder.GetProperty(GET_MEMBER_NAME_CHECKED(UCaptureManagerSettings, GeneralTokens));
+	IDetailPropertyRow* GeneralTokensRow = InDetailBuilder.EditDefaultProperty(GeneralTokensHandle);
+	check(GeneralTokensRow);
+
+	FNamingTokenFilterArgs GeneralArgs = BaseArgs;
+	if (const TObjectPtr<const UCaptureManagerGeneralTokens> Tokens = GetDefault<UCaptureManagerSettings>()->GetGeneralNamingTokens())
+	{
+		GeneralArgs.NamespacesToInclude.Add(Tokens->GetNamespace());
+	}
+	BuildSlate(GeneralTokensHandle, GeneralTokensRow, GeneralArgs);
+
+	// Video encoder tokens
+	const TSharedRef<IPropertyHandle> VideoTokensHandle = InDetailBuilder.GetProperty(GET_MEMBER_NAME_CHECKED(UCaptureManagerSettings, VideoCommandTokens));
+	IDetailPropertyRow* VideoTokensRow = InDetailBuilder.EditDefaultProperty(VideoTokensHandle);
+	check(VideoTokensRow);
+	
+	FNamingTokenFilterArgs VideoArgs = BaseArgs;
+	if (const TObjectPtr<const UCaptureManagerVideoEncoderTokens> VideoTokens = GetDefault<UCaptureManagerSettings>()->GetVideoEncoderNamingTokens())
+	{
+		VideoArgs.NamespacesToInclude.Add(VideoTokens->GetNamespace());
+	}
+	BuildSlate(VideoTokensHandle, VideoTokensRow, VideoArgs);
 	
-	const TSharedRef<IPropertyHandle> AutomaticTokensHandle = InDetailBuilder.GetProperty(GET_MEMBER_NAME_CHECKED(UCaptureManagerSettings, AutomaticTokens));
-	IDetailPropertyRow* AutomaticTokensRow = InDetailBuilder.EditDefaultProperty(AutomaticTokensHandle);
-	check(AutomaticTokensRow);
+	// Audio encoder tokens
+	const TSharedRef<IPropertyHandle> AudioTokensHandle = InDetailBuilder.GetProperty(GET_MEMBER_NAME_CHECKED(UCaptureManagerSettings, AudioCommandTokens));
+	IDetailPropertyRow* AudioTokensRow = InDetailBuilder.EditDefaultProperty(AudioTokensHandle);
+	check(AudioTokensRow);
 	
-	AutomaticTokensRow->CustomWidget()
-	.WholeRowContent()
-	[
-		SNew(SVerticalBox)
-		+ SVerticalBox::Slot()
-		.Padding(0.f, 4.f)
-		.AutoHeight()
-		[
-			SNew(STextBlock)
-			.Text(AutomaticTokensHandle->GetPropertyDisplayName())
-			.Font(IDetailLayoutBuilder::GetDetailFontBold())
-		]
-		+ SVerticalBox::Slot()
-		.AutoHeight()
-		[
-			SNew(SEditableText)
-			.IsReadOnly(true)
-			.Text_Raw(this, &FCaptureManagerSettingsCustomization::GetDisplayTokenText)
-			.Font(IDetailLayoutBuilder::GetDetailFont())
-		]
-	];
+	FNamingTokenFilterArgs AudioArgs = BaseArgs;
+	if (const TObjectPtr<const UCaptureManagerAudioEncoderTokens> AudioTokens = GetDefault<UCaptureManagerSettings>()->GetAudioEncoderNamingTokens())
+	{
+		AudioArgs.NamespacesToInclude.Add(AudioTokens->GetNamespace());
+	}
+	BuildSlate(AudioTokensHandle, AudioTokensRow, AudioArgs);
 }
 
-FText FCaptureManagerSettingsCustomization::GetDisplayTokenText() const
+FText FCaptureManagerSettingsCustomization::GetDisplayTokenText(FNamingTokenFilterArgs InArgs) const
 {
 	FString FormattedTokensString = TEXT("None");
 	if (GEditor)
 	{
-		FNamingTokenFilterArgs Args;
-		if (const TObjectPtr<const UCaptureManagerGeneralTokens> Tokens = GetDefault<UCaptureManagerSettings>()->GetGeneralTokens())
-		{
-			Args.NamespacesToInclude.Add(Tokens->GetNamespace());
-		}
-		FormattedTokensString = GEditor->GetEditorSubsystem<UNamingTokensEditorSubsystem>()->GetFormattedTokensStringForDisplay(MoveTemp(Args));
+		FormattedTokensString = GEditor->GetEditorSubsystem<UNamingTokensEditorSubsystem>()->GetFormattedTokensStringForDisplay(MoveTemp(InArgs));
 	}
 	return FText::FromString(FormattedTokensString);
 }
diff --git a/Engine/Plugins/VirtualProduction/CaptureManager/CaptureManagerApp/Source/CaptureManagerSettings/Private/Settings/CaptureManagerSettingsCustomization.h b/Engine/Plugins/VirtualProduction/CaptureManager/CaptureManagerApp/Source/CaptureManagerSettings/Private/Settings/CaptureManagerSettingsCustomization.h
--- a/Engine/Plugins/VirtualProduction/CaptureManager/CaptureManagerApp/Source/CaptureManagerSettings/Private/Settings/CaptureManagerSettingsCustomization.h
+++ b/Engine/Plugins/VirtualProduction/CaptureManager/CaptureManagerApp/Source/CaptureManagerSettings/Private/Settings/CaptureManagerSettingsCustomization.h
@@ -5,6 +5,8 @@
 #include "CoreMinimal.h"
 #include "IDetailCustomization.h"
 
+struct FNamingTokenFilterArgs;
+
 /**
  * Customization for the capture manager settings.
  */
@@ -15,5 +17,5 @@ public:
 	virtual void CustomizeDetails(IDetailLayoutBuilder& InDetailBuilder) override;
 	
 private:
-	FText GetDisplayTokenText() const;
+	FText GetDisplayTokenText(FNamingTokenFilterArgs InArgs) const;
 };
\ No newline at end of file
diff --git a/Engine/Plugins/VirtualProduction/CaptureManager/CaptureManagerApp/Source/CaptureManagerSettings/Private/Settings/CaptureManagerTemplateTokens.cpp b/Engine/Plugins/VirtualProduction/CaptureManager/CaptureManagerApp/Source/CaptureManagerSettings/Private/Settings/CaptureManagerTemplateTokens.cpp
--- a/Engine/Plugins/VirtualProduction/CaptureManager/CaptureManagerApp/Source/CaptureManagerSettings/Private/Settings/CaptureManagerTemplateTokens.cpp
+++ b/Engine/Plugins/VirtualProduction/CaptureManager/CaptureManagerApp/Source/CaptureManagerSettings/Private/Settings/CaptureManagerTemplateTokens.cpp
@@ -4,16 +4,21 @@
 
 #define LOCTEXT_NAMESPACE "CaptureManagerNamingTokens"
 
-TMap<FString, UE::CaptureManager::FArchiveToken> UCaptureManagerGeneralTokens::GeneralTokens = {
-	{FString(UE::CaptureManager::IdKey), {TEXT("id"), LOCTEXT("ArchiveId", "Archive Unique Id")}},
-	{FString(UE::CaptureManager::DeviceKey), {TEXT("device"), LOCTEXT("ArchiveDeviceId", "Archive Device User Id")}},
-	{FString(UE::CaptureManager::SlateKey), {TEXT("slate"), LOCTEXT("ArchiveSlate", "Archive Slate")}},
-	{FString(UE::CaptureManager::TakeKey), {TEXT("take"), LOCTEXT("ArchiveTake", "Archive Take")}}
-};
 
 UCaptureManagerGeneralTokens::UCaptureManagerGeneralTokens()
 {
 	Namespace = TEXT("cpman");
+
+	GeneralTokens.Emplace(FString(UE::CaptureManager::GeneralTokens::IdKey), { TEXT("id"), LOCTEXT("ArchiveId", "Archive Unique Id") });
+	GeneralTokens.Emplace(FString(UE::CaptureManager::GeneralTokens::DeviceKey), { TEXT("device"), LOCTEXT("ArchiveDeviceId", "Archive Device User Id") });
+	GeneralTokens.Emplace(FString(UE::CaptureManager::GeneralTokens::SlateKey), { TEXT("slate"), LOCTEXT("ArchiveSlate", "Archive Slate") });
+	GeneralTokens.Emplace(FString(UE::CaptureManager::GeneralTokens::TakeKey), { TEXT("take"), LOCTEXT("ArchiveTake", "Archive Take") });
+}
+
+UE::CaptureManager::FArchiveToken UCaptureManagerGeneralTokens::GetToken(const FString& InKey) const
+{
+	check(GeneralTokens.Contains(InKey));
+	return GeneralTokens[InKey];
 }
 
 void UCaptureManagerGeneralTokens::OnCreateDefaultTokens(TArray<FNamingTokenData>& OutTokens)
@@ -22,15 +27,13 @@ void UCaptureManagerGeneralTokens::OnCreateDefaultTokens(TArray<FNamingTokenData
 
 	for (const TPair<FString, UE::CaptureManager::FArchiveToken>& Token : GeneralTokens)
 	{
-		OutTokens.Emplace(
+		OutTokens.Add({
 			Token.Value.Name,
 			Token.Value.Description,
-			FNamingTokenData::FTokenProcessorDelegateNative::CreateLambda(
-				[Name = Token.Value.Name] {
+			FNamingTokenData::FTokenProcessorDelegateNative::CreateLambda([Name = Token.Value.Name] {
 					return FText::FromString(Name);
-				}
-			)
-		);
+				})
+			});
 	}
 
 }
@@ -45,4 +48,89 @@ void UCaptureManagerGeneralTokens::OnPostEvaluate_Implementation()
 	Super::OnPostEvaluate_Implementation();
 }
 
+
+UCaptureManagerVideoEncoderTokens::UCaptureManagerVideoEncoderTokens()
+{
+	Namespace = TEXT("cmvidenc");
+
+	VideoEncoderTokens.Emplace(FString(UE::CaptureManager::VideoEncoderTokens::InputKey), {TEXT("input"), LOCTEXT("VideoInputPath", "Input File Path")});
+	VideoEncoderTokens.Emplace(FString(UE::CaptureManager::VideoEncoderTokens::OutputKey), { TEXT("output"), LOCTEXT("VideoOutputPath", "Output File Path") });
+	VideoEncoderTokens.Emplace(FString(UE::CaptureManager::VideoEncoderTokens::ParamsKey), { TEXT("params"), LOCTEXT("VideoParams", "Conversion Parameters (e.g. Pixel Format, Rotation etc") });
+}
+
+UE::CaptureManager::FArchiveToken UCaptureManagerVideoEncoderTokens::GetToken(const FString& InKey) const
+{
+	check(VideoEncoderTokens.Contains(InKey));
+	return VideoEncoderTokens[InKey];
+}
+
+void UCaptureManagerVideoEncoderTokens::OnCreateDefaultTokens(TArray<FNamingTokenData>& OutTokens)
+{
+	Super::OnCreateDefaultTokens(OutTokens);
+
+	for (const TPair<FString, UE::CaptureManager::FArchiveToken>& Token : VideoEncoderTokens)
+	{
+		OutTokens.Add({
+			Token.Value.Name,
+			Token.Value.Description,
+			FNamingTokenData::FTokenProcessorDelegateNative::CreateLambda([Name = Token.Value.Name] {
+					return FText::FromString(Name);
+				})
+			});
+	}
+
+}
+
+void UCaptureManagerVideoEncoderTokens::OnPreEvaluate_Implementation(const FNamingTokensEvaluationData& InEvaluationData)
+{
+	Super::OnPreEvaluate_Implementation(InEvaluationData);
+}
+
+void UCaptureManagerVideoEncoderTokens::OnPostEvaluate_Implementation()
+{
+	Super::OnPostEvaluate_Implementation();
+}
+
+
+UCaptureManagerAudioEncoderTokens::UCaptureManagerAudioEncoderTokens()
+{
+	Namespace = TEXT("cmaudenc");
+
+	AudioEncoderTokens.Emplace(FString(UE::CaptureManager::AudioEncoderTokens::InputKey), {TEXT("input"), LOCTEXT("AudioInputPath", "Input File Path")});
+	AudioEncoderTokens.Emplace(FString(UE::CaptureManager::AudioEncoderTokens::OutputKey), { TEXT("output"), LOCTEXT("AudioOutputPath", "Output File Path") });
+}
+
+UE::CaptureManager::FArchiveToken UCaptureManagerAudioEncoderTokens::GetToken(const FString& InKey) const
+{
+	check(AudioEncoderTokens.Contains(InKey));
+	return AudioEncoderTokens[InKey];
+}
+
+void UCaptureManagerAudioEncoderTokens::OnCreateDefaultTokens(TArray<FNamingTokenData>& OutTokens)
+{
+	Super::OnCreateDefaultTokens(OutTokens);
+
+	for (const TPair<FString, UE::CaptureManager::FArchiveToken>& Token : AudioEncoderTokens)
+	{
+		OutTokens.Add({
+			Token.Value.Name,
+			Token.Value.Description,
+			FNamingTokenData::FTokenProcessorDelegateNative::CreateLambda([Name = Token.Value.Name] {
+					return FText::FromString(Name);
+				})
+			});
+	}
+
+}
+
+void UCaptureManagerAudioEncoderTokens::OnPreEvaluate_Implementation(const FNamingTokensEvaluationData& InEvaluationData)
+{
+	Super::OnPreEvaluate_Implementation(InEvaluationData);
+}
+
+void UCaptureManagerAudioEncoderTokens::OnPostEvaluate_Implementation()
+{
+	Super::OnPostEvaluate_Implementation();
+}
+
 #undef LOCTEXT_NAMESPACE
diff --git a/Engine/Plugins/VirtualProduction/CaptureManager/CaptureManagerApp/Source/CaptureManagerSettings/Public/Settings/CaptureManagerSettings.h b/Engine/Plugins/VirtualProduction/CaptureManager/CaptureManagerApp/Source/CaptureManagerSettings/Public/Settings/CaptureManagerSettings.h
--- a/Engine/Plugins/VirtualProduction/CaptureManager/CaptureManagerApp/Source/CaptureManagerSettings/Public/Settings/CaptureManagerSettings.h
+++ b/Engine/Plugins/VirtualProduction/CaptureManager/CaptureManagerApp/Source/CaptureManagerSettings/Public/Settings/CaptureManagerSettings.h
@@ -7,6 +7,8 @@
 #include "CaptureManagerSettings.generated.h"
 
 class UCaptureManagerGeneralTokens;
+class UCaptureManagerVideoEncoderTokens;
+class UCaptureManagerAudioEncoderTokens;
 
 UCLASS(config = Editor, defaultconfig)
 class CAPTUREMANAGERSETTINGS_API UCaptureManagerSettings : public UObject
@@ -16,8 +18,14 @@ public:
 
 	void Initialize();
 
-	/** Get the naming tokens for Capture Manager. */
-	TObjectPtr<const UCaptureManagerGeneralTokens> GetGeneralTokens() const;
+	/** Get the general naming tokens for Capture Manager. */
+	TObjectPtr<UCaptureManagerGeneralTokens> GetGeneralNamingTokens() const;
+
+	/** Get the video encoder naming tokens for Capture Manager. */
+	TObjectPtr<UCaptureManagerVideoEncoderTokens> GetVideoEncoderNamingTokens() const;
+
+	/** Get the audio encoder naming tokens for Capture Manager. */
+	TObjectPtr<UCaptureManagerAudioEncoderTokens> GetAudioEncoderNamingTokens() const;
 
 	/** Default location to where the converted data will be stored. It can be overriden when configuring the Ingest Job. */
 	UPROPERTY(config, EditAnywhere, Category = "Conversion")
@@ -47,6 +55,11 @@ public:
 		meta = (EditCondition = "bEnableThirdPartyEncoder", EditConditionHides))
 	FString CustomVideoCommandArguments;
 
+	/** Tokens compatible with video command properties. */
+	UPROPERTY(config, EditAnywhere, Category = "Conversion",
+		meta = (EditCondition = "bEnableThirdPartyEncoder", EditConditionHides))
+	FText VideoCommandTokens;
+
 	/** Custom audio command arguments to be used for executing the third party encoder.
 	 * NOTE: Leave empty to use the arguments generated from the Job settings.
 	 */
@@ -54,9 +67,14 @@ public:
 		meta = (EditCondition = "bEnableThirdPartyEncoder", EditConditionHides))
 	FString CustomAudioCommandArguments;
 
+	/** Tokens compatible with audio command properties. */
+	UPROPERTY(config, EditAnywhere, Category = "Conversion",
+		meta = (EditCondition = "bEnableThirdPartyEncoder", EditConditionHides))
+	FText AudioCommandTokens;
+
 	/** Placeholder for a list of the automatic tokens, set from the customization. */
 	UPROPERTY(VisibleAnywhere, Category = "Templates")
-	FText AutomaticTokens;
+	FText GeneralTokens;
 
 private:
 
@@ -69,5 +87,11 @@ private:
 	* This isn't serialized to the config file, and exists here for singleton-like access.
 	*/
 	UPROPERTY(Instanced, Transient)
-	TObjectPtr<UCaptureManagerGeneralTokens> GeneralTokens;
+	TObjectPtr<UCaptureManagerGeneralTokens> GeneralNamingTokens;
+
+	UPROPERTY(Instanced, Transient)
+	TObjectPtr<UCaptureManagerVideoEncoderTokens> VideoEncoderNamingTokens;
+
+	UPROPERTY(Instanced, Transient)
+	TObjectPtr<UCaptureManagerAudioEncoderTokens> AudioEncoderNamingTokens;
 };
\ No newline at end of file
diff --git a/Engine/Plugins/VirtualProduction/CaptureManager/CaptureManagerApp/Source/CaptureManagerSettings/Public/Settings/CaptureManagerTemplateTokens.h b/Engine/Plugins/VirtualProduction/CaptureManager/CaptureManagerApp/Source/CaptureManagerSettings/Public/Settings/CaptureManagerTemplateTokens.h
--- a/Engine/Plugins/VirtualProduction/CaptureManager/CaptureManagerApp/Source/CaptureManagerSettings/Public/Settings/CaptureManagerTemplateTokens.h
+++ b/Engine/Plugins/VirtualProduction/CaptureManager/CaptureManagerApp/Source/CaptureManagerSettings/Public/Settings/CaptureManagerTemplateTokens.h
@@ -8,20 +8,35 @@
 
 namespace UE::CaptureManager
 {
+struct FArchiveToken
+{
+	FString Name;
+	FText Description;
+};
 
+namespace GeneralTokens
+{
 static constexpr FStringView IdKey = TEXTVIEW("id");
 static constexpr FStringView DeviceKey = TEXTVIEW("device");
 static constexpr FStringView SlateKey = TEXTVIEW("slate");
 static constexpr FStringView TakeKey = TEXTVIEW("take");
+}
 
-struct FArchiveToken
+namespace VideoEncoderTokens
 {
-	FString Name;
-	FText Description;
-};
+static constexpr FStringView InputKey = TEXTVIEW("input");
+static constexpr FStringView OutputKey = TEXTVIEW("output");
+static constexpr FStringView ParamsKey = TEXTVIEW("params");
 }
 
-UCLASS()
+namespace AudioEncoderTokens
+{
+static constexpr FStringView InputKey = TEXTVIEW("input");
+static constexpr FStringView OutputKey = TEXTVIEW("output");
+}
+}
+
+UCLASS(NotBlueprintable)
 class CAPTUREMANAGERSETTINGS_API UCaptureManagerGeneralTokens final : public UNamingTokens
 {
 	GENERATED_BODY()
@@ -29,12 +44,59 @@ class CAPTUREMANAGERSETTINGS_API UCaptureManagerGeneralTokens final : public UNa
 public:
 	UCaptureManagerGeneralTokens();
 
-	static TMap<FString, UE::CaptureManager::FArchiveToken> GeneralTokens;
+	UE::CaptureManager::FArchiveToken GetToken(const FString& InKey) const;
+
+protected:
+	// ~Begin UNamingTokens
+	virtual void OnCreateDefaultTokens(TArray<FNamingTokenData>& OutTokens) override;
+	virtual void OnPreEvaluate_Implementation(const FNamingTokensEvaluationData& InEvaluationData) override;
+	virtual void OnPostEvaluate_Implementation() override;
+	// ~End UNamingTokens
+
+private:
+	TMap<FString, UE::CaptureManager::FArchiveToken> GeneralTokens;
+};
+
+
+UCLASS(NotBlueprintable)
+class CAPTUREMANAGERSETTINGS_API UCaptureManagerVideoEncoderTokens final : public UNamingTokens
+{
+	GENERATED_BODY()
+
+public:
+	UCaptureManagerVideoEncoderTokens();
+
+	UE::CaptureManager::FArchiveToken GetToken(const FString& InKey) const;
 
 protected:
 	// ~Begin UNamingTokens
-	virtual void OnCreateDefaultTokens(TArray<FNamingTokenData>& Tokens) override;
+	virtual void OnCreateDefaultTokens(TArray<FNamingTokenData>& OutTokens) override;
 	virtual void OnPreEvaluate_Implementation(const FNamingTokensEvaluationData& InEvaluationData) override;
 	virtual void OnPostEvaluate_Implementation() override;
 	// ~End UNamingTokens
+
+private:
+	TMap<FString, UE::CaptureManager::FArchiveToken> VideoEncoderTokens;
+};
+
+
+UCLASS(NotBlueprintable)
+class CAPTUREMANAGERSETTINGS_API UCaptureManagerAudioEncoderTokens final : public UNamingTokens
+{
+	GENERATED_BODY()
+
+public:
+	UCaptureManagerAudioEncoderTokens();
+
+	UE::CaptureManager::FArchiveToken GetToken(const FString& InKey) const;
+
+protected:
+	// ~Begin UNamingTokens
+	virtual void OnCreateDefaultTokens(TArray<FNamingTokenData>& OutTokens) override;
+	virtual void OnPreEvaluate_Implementation(const FNamingTokensEvaluationData& InEvaluationData) override;
+	virtual void OnPostEvaluate_Implementation() override;
+	// ~End UNamingTokens
+
+private:
+	TMap<FString, UE::CaptureManager::FArchiveToken> AudioEncoderTokens;
 };
diff --git a/Engine/Plugins/VirtualProduction/CaptureManager/CaptureManagerEditor/Source/LiveLinkHubWorkerManager/Private/IngestProcess/IngestCaptureDataProcess.cpp b/Engine/Plugins/VirtualProduction/CaptureManager/CaptureManagerEditor/Source/LiveLinkHubWorkerManager/Private/IngestProcess/IngestCaptureDataProcess.cpp
--- a/Engine/Plugins/VirtualProduction/CaptureManager/CaptureManagerEditor/Source/LiveLinkHubWorkerManager/Private/IngestProcess/IngestCaptureDataProcess.cpp
+++ b/Engine/Plugins/VirtualProduction/CaptureManager/CaptureManagerEditor/Source/LiveLinkHubWorkerManager/Private/IngestProcess/IngestCaptureDataProcess.cpp
@@ -118,13 +118,16 @@ UE::CaptureManager::FCreateAssetsData FIngestCaptureDataProcess::PrepareAssetsDa
 
 	const UCaptureManagerEditorSettings* Settings = GetDefault<UCaptureManagerEditorSettings>();
 
+	const TObjectPtr<const UCaptureManagerIngestNamingTokens> Tokens = Settings->GetGeneralNamingTokens();
+	check(Tokens);
+
 	FStringFormatNamedArguments ImportNamedArgs;
 	{
 		// Editor tokens
-		ImportNamedArgs.Add(Settings->GetGeneralNamingTokens()->GetToken(FString(GeneralTokens::IdKey)).Name, InTakeUploadId.ToString());
-		ImportNamedArgs.Add(Settings->GetGeneralNamingTokens()->GetToken(FString(GeneralTokens::DeviceKey)).Name, InDeviceName);
-		ImportNamedArgs.Add(Settings->GetGeneralNamingTokens()->GetToken(FString(GeneralTokens::SlateKey)).Name, InIngestCaptureData.Slate);
-		ImportNamedArgs.Add(Settings->GetGeneralNamingTokens()->GetToken(FString(GeneralTokens::TakeKey)).Name, FString::FromInt(InIngestCaptureData.TakeNumber));
+		ImportNamedArgs.Add(Tokens->GetToken(FString(GeneralTokens::IdKey)).Name, InTakeUploadId.ToString());
+		ImportNamedArgs.Add(Tokens->GetToken(FString(GeneralTokens::DeviceKey)).Name, InDeviceName);
+		ImportNamedArgs.Add(Tokens->GetToken(FString(GeneralTokens::SlateKey)).Name, InIngestCaptureData.Slate);
+		ImportNamedArgs.Add(Tokens->GetToken(FString(GeneralTokens::TakeKey)).Name, FString::FromInt(InIngestCaptureData.TakeNumber));
 	}
 
 	int32 TakeId = 0;
