[CaptureManager] Fix evaluation of Download Directory

diff --git a/Engine/Plugins/VirtualProduction/CaptureManager/CaptureManagerApp/Source/CaptureManagerEditor/Private/IngestManagement/IngestJob.h b/Engine/Plugins/VirtualProduction/CaptureManager/CaptureManagerApp/Source/CaptureManagerEditor/Private/IngestManagement/IngestJob.h
--- a/Engine/Plugins/VirtualProduction/CaptureManager/CaptureManagerApp/Source/CaptureManagerEditor/Private/IngestManagement/IngestJob.h
+++ b/Engine/Plugins/VirtualProduction/CaptureManager/CaptureManagerApp/Source/CaptureManagerEditor/Private/IngestManagement/IngestJob.h
@@ -37,6 +37,7 @@ public:
 		};
 
 		FString LocalWorkFolder;
+		FString DownloadFolder;
 		FVideoSettings VideoSettings;
 		FAudioSettings AudioSettings;
 	};
diff --git a/Engine/Plugins/VirtualProduction/CaptureManager/CaptureManagerApp/Source/CaptureManagerEditor/Private/IngestManagement/IngestJobExecutor.cpp b/Engine/Plugins/VirtualProduction/CaptureManager/CaptureManagerApp/Source/CaptureManagerEditor/Private/IngestManagement/IngestJobExecutor.cpp
--- a/Engine/Plugins/VirtualProduction/CaptureManager/CaptureManagerApp/Source/CaptureManagerEditor/Private/IngestManagement/IngestJobExecutor.cpp
+++ b/Engine/Plugins/VirtualProduction/CaptureManager/CaptureManagerApp/Source/CaptureManagerEditor/Private/IngestManagement/IngestJobExecutor.cpp
@@ -194,6 +194,10 @@ UIngestCapability_ProcessHandle* FIngestJobExecutor::RunIngest(FIngestJob& InJob
 			FNamingTokenResultData WorkingDirectoryResult = NamingTokensSubsystem->EvaluateTokenString(WorkingDirectory, GeneralNamingTokenArgs);
 			IngestOptions->WorkingDirectory = WorkingDirectoryResult.EvaluatedText.ToString();
 
+			FString DownloadDirectory = FString::Format(*Settings.DownloadFolder, GeneralArgs);
+			FNamingTokenResultData DownloadDirectoryResult = NamingTokensSubsystem->EvaluateTokenString(DownloadDirectory, GeneralNamingTokenArgs);
+			IngestOptions->DownloadDirectory = DownloadDirectoryResult.EvaluatedText.ToString();
+
 			FString FileNamePrefix = FString::Format(*Settings.VideoSettings.FileNamePrefix, GeneralArgs);
 			FNamingTokenResultData VideoPrefixResult = NamingTokensSubsystem->EvaluateTokenString(FileNamePrefix, GeneralNamingTokenArgs);
 			FIngestCapability_VideoOptions VideoOptions;
diff --git a/Engine/Plugins/VirtualProduction/CaptureManager/CaptureManagerApp/Source/LiveLinkCapabilities/Public/Ingest/IngestCapability_Options.h b/Engine/Plugins/VirtualProduction/CaptureManager/CaptureManagerApp/Source/LiveLinkCapabilities/Public/Ingest/IngestCapability_Options.h
--- a/Engine/Plugins/VirtualProduction/CaptureManager/CaptureManagerApp/Source/LiveLinkCapabilities/Public/Ingest/IngestCapability_Options.h
+++ b/Engine/Plugins/VirtualProduction/CaptureManager/CaptureManagerApp/Source/LiveLinkCapabilities/Public/Ingest/IngestCapability_Options.h
@@ -75,6 +75,9 @@ public:
 	UPROPERTY(EditAnywhere, BlueprintReadWrite, Category = "Live Link Device|Ingest")
 	FString WorkingDirectory;
 
+	UPROPERTY(EditAnywhere, BlueprintReadWrite, Category = "Live Link Device|Ingest")
+	FString DownloadDirectory;
+
 	UPROPERTY(EditAnywhere, BlueprintReadWrite, Category = "Live Link Device|Ingest")
 	FIngestCapability_VideoOptions Video;
 
