[CaptureManager] Propagate download directory setting and fix incorrect device name in working directory

diff --git a/Engine/Plugins/VirtualProduction/CaptureManager/CaptureManagerApp/Source/CaptureManagerEditor/Private/IngestManagement/IngestJobExecutor.cpp b/Engine/Plugins/VirtualProduction/CaptureManager/CaptureManagerApp/Source/CaptureManagerEditor/Private/IngestManagement/IngestJobExecutor.cpp
--- a/Engine/Plugins/VirtualProduction/CaptureManager/CaptureManagerApp/Source/CaptureManagerEditor/Private/IngestManagement/IngestJobExecutor.cpp
+++ b/Engine/Plugins/VirtualProduction/CaptureManager/CaptureManagerApp/Source/CaptureManagerEditor/Private/IngestManagement/IngestJobExecutor.cpp
@@ -164,11 +164,34 @@ UIngestCapability_ProcessHandle* FIngestJobExecutor::RunIngest(FIngestJob& InJob
 		return nullptr;
 	}
 
+	ULiveLinkDeviceSubsystem* DeviceSubsystem = GEngine->GetEngineSubsystem<ULiveLinkDeviceSubsystem>();
+	check(DeviceSubsystem);
+
+	const TObjectPtr<ULiveLinkDevice>* DevicePtr = DeviceSubsystem->GetDeviceMap().Find(InJob.GetCaptureDeviceId());
+
+	if (!DevicePtr)
+	{
+		UE_LOG(
+			LogCaptureManagerIngestJobExecutor,
+			Error,
+			TEXT("[%s] Ingest failed, Capture Device does not exist"), *ExecutorName
+		);
+
+		// Set the processing state to anything other than Pending, to prevent other executors picking up the same bad job.
+		InJob.SetProcessingState(FIngestJob::EProcessingState::Aborted);
+		InPromise.SetValue();
+		return nullptr;
+	}
+
+	TObjectPtr<ULiveLinkDevice> Device = *DevicePtr;
+	check(Device);
+	const FText DeviceName = Device->GetDisplayName();
+
 	UIngestCapability_Options* IngestOptions = NewObject<UIngestCapability_Options>();
 
 	// Naming tokens subsystem consults asset registry so need to run on the game thread
 	CallOnGameThread(
-		[&InJob, IngestOptions]()
+		[&InJob, &DeviceName, IngestOptions]()
 		{
 			UNamingTokensEditorSubsystem* NamingTokensSubsystem = GEditor->GetEditorSubsystem<UNamingTokensEditorSubsystem>();
 			check(NamingTokensSubsystem);
@@ -185,7 +208,7 @@ UIngestCapability_ProcessHandle* FIngestJobExecutor::RunIngest(FIngestJob& InJob
 			{
 				using namespace UE::CaptureManager;
 				GeneralArgs.Add(Tokens->GetToken(FString(GeneralTokens::IdKey)).Name, InJob.GetTakeMetadata().UniqueId);
-				GeneralArgs.Add(Tokens->GetToken(FString(GeneralTokens::DeviceKey)).Name, InJob.GetTakeMetadata().Device.UserId);
+				GeneralArgs.Add(Tokens->GetToken(FString(GeneralTokens::DeviceKey)).Name, DeviceName.ToString());
 				GeneralArgs.Add(Tokens->GetToken(FString(GeneralTokens::SlateKey)).Name, InJob.GetTakeMetadata().Slate);
 				GeneralArgs.Add(Tokens->GetToken(FString(GeneralTokens::TakeKey)).Name, InJob.GetTakeMetadata().TakeNumber);
 			}
@@ -216,26 +239,6 @@ UIngestCapability_ProcessHandle* FIngestJobExecutor::RunIngest(FIngestJob& InJob
 		}
 	);
 
-	ULiveLinkDeviceSubsystem* DeviceSubsystem = GEngine->GetEngineSubsystem<ULiveLinkDeviceSubsystem>();
-
-	const TObjectPtr<ULiveLinkDevice>* DevicePtr = DeviceSubsystem->GetDeviceMap().Find(InJob.GetCaptureDeviceId());
-
-	if (!DevicePtr)
-	{
-		UE_LOG(
-			LogCaptureManagerIngestJobExecutor,
-			Error,
-			TEXT("[%s] Ingest failed, Capture Device does not exist"), *ExecutorName
-		);
-
-		// Set the processing state to anything other than Pending, to prevent other executors picking up the same bad job.
-		InJob.SetProcessingState(FIngestJob::EProcessingState::Aborted);
-		InPromise.SetValue();
-		return nullptr;
-	}
-
-	TObjectPtr<ULiveLinkDevice> Device = *DevicePtr;
-
 	if (!Device->Implements<ULiveLinkDeviceCapability_Ingest>())
 	{
 		UE_LOG(
diff --git a/Engine/Plugins/VirtualProduction/CaptureManager/CaptureManagerApp/Source/GenericLiveLinkDevice/Private/CPS/GenericLiveLinkCPSDevice.cpp b/Engine/Plugins/VirtualProduction/CaptureManager/CaptureManagerApp/Source/GenericLiveLinkDevice/Private/CPS/GenericLiveLinkCPSDevice.cpp
--- a/Engine/Plugins/VirtualProduction/CaptureManager/CaptureManagerApp/Source/GenericLiveLinkDevice/Private/CPS/GenericLiveLinkCPSDevice.cpp
+++ b/Engine/Plugins/VirtualProduction/CaptureManager/CaptureManagerApp/Source/GenericLiveLinkDevice/Private/CPS/GenericLiveLinkCPSDevice.cpp
@@ -116,13 +116,13 @@ void UGenericLiveLinkCPSDevice::RunIngestTake(int32 InTakeId, const UIngestCapab
 		TotalSize += File.Length;
 	}
 
-	const FString DownloadedStorage = GetDefault<UCaptureManagerSettings>()->DownloadDirectory.Path;
+	FString DownloadedStorage = InIngestOptions->DownloadDirectory;
 
 	FTaskProgress::FTask Task = TaskProgress->StartTask();
 
 	TStrongObjectPtr<const UIngestCapability_Options> IngestOptions(InIngestOptions);
 
-	TUniquePtr<FCPSFileStream> CPSStream = MakeUnique<FCPSFileStream>(DownloadedStorage, TotalSize);
+	TUniquePtr<FCPSFileStream> CPSStream = MakeUnique<FCPSFileStream>(MoveTemp(DownloadedStorage), TotalSize);
 	CPSStream->SetExportFinished(FCPSFileStream::FExportFinished::CreateUObject(this, &UGenericLiveLinkCPSDevice::OnExportFinished, Take.Name, InTakeId, TaskProgress, IngestOptions));
 	CPSStream->SetProgressHandler(FCPSFileStream::FReportProgress::CreateUObject(this, &UGenericLiveLinkCPSDevice::OnExportProgressReport, Task));
 
@@ -375,13 +375,13 @@ void UGenericLiveLinkCPSDevice::OnExportFinished(UE::CaptureManager::TProtocolRe
 {
 	using namespace UE::CaptureManager;
 
-	const FString DownloadedStorage = GetDefault<UCaptureManagerSettings>()->DownloadDirectory.Path;
+	const FString DownloadedStorage = InIngestOptions->DownloadDirectory;
 	const FString DownloadedTake = FPaths::Combine(DownloadedStorage, InTakeName);
 
 	if (InResult.HasValue())
 	{
 		// Blocking function
-		TOptional<FTakeMetadata> TakeMetadata = ParseTake(InTakeName);
+		TOptional<FTakeMetadata> TakeMetadata = ParseTake(DownloadedStorage, InTakeName);
 		if (TakeMetadata.IsSet())
 		{
 			FScopeLock Lock(&DownloadedTakesMutex);
@@ -407,21 +407,21 @@ void UGenericLiveLinkCPSDevice::OnExportFinished(UE::CaptureManager::TProtocolRe
 	IFileManager::Get().DeleteDirectory(*DownloadedTake, false, true);
 }
 
-TOptional<FTakeMetadata> UGenericLiveLinkCPSDevice::ParseTake(const FString& InTakeName)
+TOptional<FTakeMetadata> UGenericLiveLinkCPSDevice::ParseTake(const FString& InTakeDirectory, const FString& InTakeName)
 {
 	using namespace UE::CaptureManager;
 
-	const FString DownloadedStorage = GetDefault<UCaptureManagerSettings>()->DownloadDirectory.Path;
-	const FString DownloadedTake = FPaths::Combine(DownloadedStorage, InTakeName);
+	const FString TakeDirectory = InTakeDirectory;
+	const FString TakePath = FPaths::Combine(TakeDirectory, InTakeName);
 
 	TArray<FString> TakeFiles;
-	IFileManager::Get().FindFiles(TakeFiles, *DownloadedTake, *FTakeMetadata::FileExtension);
+	IFileManager::Get().FindFiles(TakeFiles, *TakePath, *FTakeMetadata::FileExtension);
 	if (!TakeFiles.IsEmpty())
 	{
 		check(TakeFiles.Num() == 1);
 
 		FTakeMetadataParser TakeMetadataParser;
-		TValueOrError<FTakeMetadata, FTakeMetadataParserError> TakeMetadataResult = TakeMetadataParser.Parse(DownloadedTake / TakeFiles[0]);
+		TValueOrError<FTakeMetadata, FTakeMetadataParserError> TakeMetadataResult = TakeMetadataParser.Parse(TakePath / TakeFiles[0]);
 		if (!TakeMetadataResult.HasError())
 		{
 			return TakeMetadataResult.StealValue();
@@ -429,14 +429,14 @@ TOptional<FTakeMetadata> UGenericLiveLinkCPSDevice::ParseTake(const FString& InT
 	}
 
 	TArray<FText> ValidationFailures;
-	TOptional<FTakeMetadata> ParseOldMetadataResult = LiveLinkMetadata::ParseOldLiveLinkTakeMetadata(DownloadedTake, ValidationFailures);
+	TOptional<FTakeMetadata> ParseOldMetadataResult = LiveLinkMetadata::ParseOldLiveLinkTakeMetadata(TakePath, ValidationFailures);
 
 	if (ParseOldMetadataResult.IsSet())
 	{
 		return ParseOldMetadataResult.GetValue();
 	}
 
-	ParseOldMetadataResult = StereoCameraMetadata::ParseOldStereoCameraMetadata(DownloadedTake, ValidationFailures);
+	ParseOldMetadataResult = StereoCameraMetadata::ParseOldStereoCameraMetadata(TakePath, ValidationFailures);
 
 	if (ParseOldMetadataResult.IsSet())
 	{
diff --git a/Engine/Plugins/VirtualProduction/CaptureManager/CaptureManagerApp/Source/GenericLiveLinkDevice/Private/CPS/GenericLiveLinkCPSDevice.h b/Engine/Plugins/VirtualProduction/CaptureManager/CaptureManagerApp/Source/GenericLiveLinkDevice/Private/CPS/GenericLiveLinkCPSDevice.h
--- a/Engine/Plugins/VirtualProduction/CaptureManager/CaptureManagerApp/Source/GenericLiveLinkDevice/Private/CPS/GenericLiveLinkCPSDevice.h
+++ b/Engine/Plugins/VirtualProduction/CaptureManager/CaptureManagerApp/Source/GenericLiveLinkDevice/Private/CPS/GenericLiveLinkCPSDevice.h
@@ -96,7 +96,7 @@ private:
 						  TSharedPtr<UE::CaptureManager::FTaskProgress> InTaskProgress,
 						  TStrongObjectPtr<const UIngestCapability_Options> InIngestOptions);
 
-	TOptional<FTakeMetadata> ParseTake(const FString& InTakeName);
+	TOptional<FTakeMetadata> ParseTake(const FString& InTakeDirectory, const FString& InTakeName);
 	FTakeMetadata ParseTakeMetadata(const UE::CaptureManager::FGetTakeMetadataResponse::FTakeObject& InTake);
 
 	void FetchThumbnail(const FString& InTakeName, UE::CaptureManager::FTakeId InTakeId);
