[CaptureManager] Fix for ingest not working on CPS devices

On behalf of brent.bannister.

#rnx

diff --git a/Engine/Plugins/VirtualProduction/CaptureManager/CaptureManagerApp/Source/CaptureManagerEditor/Private/IngestManagement/UIngestJobSettings.cpp b/Engine/Plugins/VirtualProduction/CaptureManager/CaptureManagerApp/Source/CaptureManagerEditor/Private/IngestManagement/UIngestJobSettings.cpp
--- a/Engine/Plugins/VirtualProduction/CaptureManager/CaptureManagerApp/Source/CaptureManagerEditor/Private/IngestManagement/UIngestJobSettings.cpp
+++ b/Engine/Plugins/VirtualProduction/CaptureManager/CaptureManagerApp/Source/CaptureManagerEditor/Private/IngestManagement/UIngestJobSettings.cpp
@@ -17,6 +17,7 @@ UDefaultJobPreset::UDefaultJobPreset(const FObjectInitializer& ObjectInitializer
 	: Super(ObjectInitializer)
 {
 	DisplayName = LOCTEXT("DefaultText", "Default");
+	DownloadFolder = GetDefault<UCaptureManagerSettings>()->DownloadDirectory;
 	LocalWorkFolder = GetDefault<UCaptureManagerSettings>()->DefaultConversionDirectory;
 	ImageFormat = EOutputImageFormat::JPEG;
 	ImageFileNamePrefix = TEXT("frame");
@@ -34,6 +35,11 @@ void UDefaultJobPreset::PostEditChangeProperty(struct FPropertyChangedEvent& InP
 	{
 		LocalWorkFolder = GetDefault<UCaptureManagerSettings>()->DefaultConversionDirectory;
 	}
+
+	if (DownloadFolder.Path.IsEmpty())
+	{
+		DownloadFolder = GetDefault<UCaptureManagerSettings>()->DownloadDirectory;
+	}
 }
 
 #undef LOCTEXT_NAMESPACE
diff --git a/Engine/Plugins/VirtualProduction/CaptureManager/CaptureManagerApp/Source/CaptureManagerEditor/Private/IngestManagement/UIngestJobSettings.h b/Engine/Plugins/VirtualProduction/CaptureManager/CaptureManagerApp/Source/CaptureManagerEditor/Private/IngestManagement/UIngestJobSettings.h
--- a/Engine/Plugins/VirtualProduction/CaptureManager/CaptureManagerApp/Source/CaptureManagerEditor/Private/IngestManagement/UIngestJobSettings.h
+++ b/Engine/Plugins/VirtualProduction/CaptureManager/CaptureManagerApp/Source/CaptureManagerEditor/Private/IngestManagement/UIngestJobSettings.h
@@ -50,6 +50,9 @@ public:
 	UPROPERTY(AdvancedDisplay)
 	FText DisplayName;
 
+	UPROPERTY(AdvancedDisplay)
+	FDirectoryPath DownloadFolder;
+
 	UPROPERTY(Category = Conversion, EditAnywhere)
 	FDirectoryPath LocalWorkFolder;
 
diff --git a/Engine/Plugins/VirtualProduction/CaptureManager/CaptureManagerApp/Source/CaptureManagerEditor/Private/UI/TakesView.cpp b/Engine/Plugins/VirtualProduction/CaptureManager/CaptureManagerApp/Source/CaptureManagerEditor/Private/UI/TakesView.cpp
--- a/Engine/Plugins/VirtualProduction/CaptureManager/CaptureManagerApp/Source/CaptureManagerEditor/Private/UI/TakesView.cpp
+++ b/Engine/Plugins/VirtualProduction/CaptureManager/CaptureManagerApp/Source/CaptureManagerEditor/Private/UI/TakesView.cpp
@@ -265,6 +265,7 @@ void FTakesView::AddTakesToIngestQueue(const TArray<TStrongObjectPtr<UTakeVirtua
 		FIngestJob::FSettings IngestJobSettings =
 		{
 			.LocalWorkFolder = InJobPreset.LocalWorkFolder.Path,
+			.DownloadFolder = InJobPreset.DownloadFolder.Path,
 			.VideoSettings =
 			{
 				.Format = StaticEnum<EOutputImageFormat>()->GetDisplayNameTextByValue(static_cast<std::underlying_type_t<EOutputImageFormat>>(InJobPreset.ImageFormat)).ToString(),
diff --git a/Engine/Plugins/VirtualProduction/CaptureManager/CaptureManagerApp/Source/IngestLiveLinkDevice/Private/BaseIngestLiveLinkDevice.cpp b/Engine/Plugins/VirtualProduction/CaptureManager/CaptureManagerApp/Source/IngestLiveLinkDevice/Private/BaseIngestLiveLinkDevice.cpp
--- a/Engine/Plugins/VirtualProduction/CaptureManager/CaptureManagerApp/Source/IngestLiveLinkDevice/Private/BaseIngestLiveLinkDevice.cpp
+++ b/Engine/Plugins/VirtualProduction/CaptureManager/CaptureManagerApp/Source/IngestLiveLinkDevice/Private/BaseIngestLiveLinkDevice.cpp
@@ -74,7 +74,8 @@ void UBaseIngestLiveLinkDevice::FImpl::IngestTake(int32 InTakeId, const UIngestC
 
 	if (TakeFullPath.IsEmpty())
 	{
-		FString Message = TEXT("Failed to find the take directory for the specified take");
+		check(false);
+		FString Message = FString::Printf(TEXT("Failed to lookup take directory for take id %d"), InTakeId);
 		Owner->ExecuteProcessFinishedReporter(InTakeId, MakeError(FIngestCapability_Error::InvalidArgument, MoveTemp(Message)));
 		return;
 	}
@@ -82,7 +83,8 @@ void UBaseIngestLiveLinkDevice::FImpl::IngestTake(int32 InTakeId, const UIngestC
 	TOptional<FTakeMetadata> MaybeTakeMetadata = Owner->GetTakeMetadata(InTakeId);
 	if (!MaybeTakeMetadata.IsSet())
 	{
-		FString Message = TEXT("Failed to find the take metadata for the specified take");
+		check(false);
+		FString Message = FString::Printf(TEXT("Failed to lookup take metadata for take id %d (TakePath=%s)"), InTakeId, *TakeFullPath);
 		Owner->ExecuteProcessFinishedReporter(InTakeId, MakeError(FIngestCapability_Error::InvalidArgument, MoveTemp(Message)));
 		return;
 	}
