[CaptureManager] Rename Capture Manager live link devices

diff --git a/Engine/Plugins/VirtualProduction/CaptureManager/CaptureManagerApp/Source/GenericLiveLinkDevice/GenericLiveLinkDevice.Build.cs b/Engine/Plugins/VirtualProduction/CaptureManager/CaptureManagerApp/Source/GenericLiveLinkDevice/GenericLiveLinkDevice.Build.cs
--- a/Engine/Plugins/VirtualProduction/CaptureManager/CaptureManagerApp/Source/GenericLiveLinkDevice/GenericLiveLinkDevice.Build.cs
+++ b/Engine/Plugins/VirtualProduction/CaptureManager/CaptureManagerApp/Source/GenericLiveLinkDevice/GenericLiveLinkDevice.Build.cs
@@ -39,7 +39,7 @@ public class GenericLiveLinkDevice : ModuleRules
 
 		PrivateIncludePaths.AddRange(new string[]
 		{
-			Path.Combine(ModuleDirectory, "Private/Archive"),
+			Path.Combine(ModuleDirectory, "Private/Filesystem"),
 			Path.Combine(ModuleDirectory, "Private/Tests")
 		});
 
diff --git a/Engine/Plugins/VirtualProduction/CaptureManager/CaptureManagerApp/Source/GenericLiveLinkDevice/Private/Archive/GenericLiveLinkArchiveDevice.cpp b/Engine/Plugins/VirtualProduction/CaptureManager/CaptureManagerApp/Source/GenericLiveLinkDevice/Private/Archive/GenericLiveLinkArchiveDevice.cpp
deleted file mode 100644
diff --git a/Engine/Plugins/VirtualProduction/CaptureManager/CaptureManagerApp/Source/GenericLiveLinkDevice/Private/Archive/GenericLiveLinkArchiveDevice.h b/Engine/Plugins/VirtualProduction/CaptureManager/CaptureManagerApp/Source/GenericLiveLinkDevice/Private/Archive/GenericLiveLinkArchiveDevice.h
deleted file mode 100644
diff --git a/Engine/Plugins/VirtualProduction/CaptureManager/CaptureManagerApp/Source/GenericLiveLinkDevice/Private/CPS/Customizations/ToggleConnectActionCustomization.cpp b/Engine/Plugins/VirtualProduction/CaptureManager/CaptureManagerApp/Source/GenericLiveLinkDevice/Private/CPS/Customizations/ToggleConnectActionCustomization.cpp
--- a/Engine/Plugins/VirtualProduction/CaptureManager/CaptureManagerApp/Source/GenericLiveLinkDevice/Private/CPS/Customizations/ToggleConnectActionCustomization.cpp
+++ b/Engine/Plugins/VirtualProduction/CaptureManager/CaptureManagerApp/Source/GenericLiveLinkDevice/Private/CPS/Customizations/ToggleConnectActionCustomization.cpp
@@ -11,7 +11,7 @@
 #include "LiveLinkDeviceSubsystem.h"
 #include "Engine/Engine.h"
 
-#include "CPS/GenericLiveLinkCPSDevice.h"
+#include "CPS/LiveLinkFaceDevice.h"
 
 #include "LiveLinkDeviceCapability_Connection.h"
 
@@ -136,7 +136,7 @@ bool FToggleConnectActionCustomization::IsConnectButtonEnabled() const
 		return false;
 	}
 
-	if (CastChecked<UGenericLiveLinkCPSDevice>(Device)->GetSettings()->IpAddress.IpAddressString.IsEmpty())
+	if (CastChecked<ULiveLinkFaceDevice>(Device)->GetSettings()->IpAddress.IpAddressString.IsEmpty())
 	{
 		return false;
 	}
diff --git a/Engine/Plugins/VirtualProduction/CaptureManager/CaptureManagerApp/Source/GenericLiveLinkDevice/Private/CPS/GenericLiveLinkCPSDevice.cpp b/Engine/Plugins/VirtualProduction/CaptureManager/CaptureManagerApp/Source/GenericLiveLinkDevice/Private/CPS/GenericLiveLinkCPSDevice.cpp
deleted file mode 100644
diff --git a/Engine/Plugins/VirtualProduction/CaptureManager/CaptureManagerApp/Source/GenericLiveLinkDevice/Private/CPS/GenericLiveLinkCPSDevice.h b/Engine/Plugins/VirtualProduction/CaptureManager/CaptureManagerApp/Source/GenericLiveLinkDevice/Private/CPS/GenericLiveLinkCPSDevice.h
deleted file mode 100644
diff --git a/Engine/Plugins/VirtualProduction/CaptureManager/CaptureManagerApp/Source/GenericLiveLinkDevice/Private/CPS/LiveLinkFaceDevice.cpp b/Engine/Plugins/VirtualProduction/CaptureManager/CaptureManagerApp/Source/GenericLiveLinkDevice/Private/CPS/LiveLinkFaceDevice.cpp
new file mode 100644
--- /dev/null
+++ b/Engine/Plugins/VirtualProduction/CaptureManager/CaptureManagerApp/Source/GenericLiveLinkDevice/Private/CPS/LiveLinkFaceDevice.cpp
@@ -0,0 +1,491 @@
+// Copyright Epic Games, Inc. All Rights Reserved.
+
+#include "LiveLinkFaceDevice.h"
+
+#include "Settings/CaptureManagerSettings.h"
+
+#include "ILiveLinkRecordingSessionInfo.h"
+
+#include "Control/Messages/Constants.h"
+#include "Control/Messages/ControlUpdate.h"
+
+#include "CPS/Protocol/CPSDataStream.h"
+#include "CPS/Protocol/CPSFileStream.h"
+
+#include "LiveLinkFaceMetadata.h"
+#include "StereoCameraMetadataParseUtils.h"
+
+#include "HAL/FileManager.h"
+
+#include "Utils/GenericLiveLinkDeviceUtility.h"
+
+const ULiveLinkFaceDeviceSettings* ULiveLinkFaceDevice::GetSettings() const
+{
+	return CastChecked<ULiveLinkFaceDeviceSettings>(Settings);
+}
+
+TSubclassOf<ULiveLinkDeviceSettings> ULiveLinkFaceDevice::StaticGetSettingsClass() const
+{
+	return ULiveLinkFaceDeviceSettings::StaticClass();
+}
+
+FText ULiveLinkFaceDevice::GetDisplayName() const
+{
+	return FText::FromString(GetSettings()->DisplayName);
+}
+
+EDeviceHealth ULiveLinkFaceDevice::GetDeviceHealth() const
+{
+	return EDeviceHealth::Nominal;
+}
+
+FText ULiveLinkFaceDevice::GetHealthText() const
+{
+	return FText::FromString("Example Health");
+}
+
+void ULiveLinkFaceDevice::OnDeviceAdded(const FGuid InDeviceGuid, ULiveLinkDeviceSettings* InSettings)
+{
+	CastChecked<ULiveLinkFaceDeviceSettings>(InSettings)->ConnectAction.DeviceGuid = InDeviceGuid;
+
+	Super::OnDeviceAdded(InDeviceGuid, InSettings);
+}
+
+void ULiveLinkFaceDevice::OnDeviceRemoved()
+{
+	ILiveLinkDeviceCapability_Connection::Execute_Disconnect(this);
+
+	Super::OnDeviceRemoved();
+}
+
+FString ULiveLinkFaceDevice::GetFullTakePath(UE::CaptureManager::FTakeId InTakeId) const
+{
+	FScopeLock Lock(&DownloadedTakesMutex);
+	if (const FString* DownloadedTake = DownloadedTakes.Find(InTakeId))
+	{
+		return *DownloadedTake;
+	}
+	
+	return FString();
+}
+
+void ULiveLinkFaceDevice::UpdateTakeList_Implementation()
+{
+	using namespace UE::CaptureManager;
+
+	RemoveAllTakes();
+
+	TProtocolResult<TArray<FGetTakeMetadataResponse::FTakeObject>> TakesResult = Device->FetchTakeList();
+
+	if (TakesResult.HasError())
+	{
+		return;
+	}
+
+	TMap<FString, FTakeId> NameToIdMap;
+
+	for (const FGetTakeMetadataResponse::FTakeObject& Take : TakesResult.StealValue())
+	{
+		FTakeId TakeId = AddTake(ParseTakeMetadata(Take));
+		Device->AddTakeMetadata(TakeId, Take);
+		NameToIdMap.Add(Take.Name, TakeId);
+	}
+
+	FetchThumbnails(MoveTemp(NameToIdMap));
+
+	ExecuteUpdateTakeListCallback(Execute_GetTakeIdentifiers(this));
+}
+
+void ULiveLinkFaceDevice::RunIngestTake(int32 InTakeId, const UIngestCapability_Options* InIngestOptions)
+{
+	using namespace UE::CaptureManager;
+
+	const FGetTakeMetadataResponse::FTakeObject& Take = Device->GetTake(InTakeId);
+
+	static constexpr uint32 NumberOfTasks = 3; // Download, Convert, Upload
+
+	TSharedPtr<FTaskProgress> TaskProgress = MakeShared<FTaskProgress>
+		(NumberOfTasks, FTaskProgress::FProgressReporter::CreateLambda([this, TakeId = InTakeId](double InProgress)
+	{
+		ExecuteProcessProgressReporter(TakeId, InProgress);
+	}));
+
+	uint64 TotalSize = 0;
+	for (const FGetTakeMetadataResponse::FFileObject& File : Take.Files)
+	{
+		TotalSize += File.Length;
+	}
+
+	FString DownloadedStorage = InIngestOptions->DownloadDirectory;
+
+	FTaskProgress::FTask Task = TaskProgress->StartTask();
+
+	TStrongObjectPtr<const UIngestCapability_Options> IngestOptions(InIngestOptions);
+
+	TUniquePtr<FCPSFileStream> CPSStream = MakeUnique<FCPSFileStream>(MoveTemp(DownloadedStorage), TotalSize);
+	CPSStream->SetExportFinished(FCPSFileStream::FExportFinished::CreateUObject(this, &ULiveLinkFaceDevice::OnExportFinished, Take.Name, InTakeId, TaskProgress, IngestOptions));
+	CPSStream->SetProgressHandler(FCPSFileStream::FReportProgress::CreateUObject(this, &ULiveLinkFaceDevice::OnExportProgressReport, Task));
+
+	Device->StartExport(InTakeId, MoveTemp(CPSStream));
+}
+
+void ULiveLinkFaceDevice::CancelIngest(int32 InTakeId)
+{
+	Device->CancelExport(InTakeId);
+
+	Super::CancelIngest(InTakeId);
+}
+
+ELiveLinkDeviceConnectionStatus ULiveLinkFaceDevice::GetConnectionStatus_Implementation() const
+{
+	if (!Device)
+	{
+		return ELiveLinkDeviceConnectionStatus::Disconnected;
+	}
+
+	if (bIsConnecting)
+	{
+		return ELiveLinkDeviceConnectionStatus::Connecting;
+	}
+	else if (Device->IsConnected())
+	{
+		return ELiveLinkDeviceConnectionStatus::Connected;
+	}
+
+	return ELiveLinkDeviceConnectionStatus::Disconnected;
+}
+
+FString ULiveLinkFaceDevice::GetHardwareId_Implementation() const
+{
+	return GetSettings()->IpAddress.IpAddressString;
+}
+
+bool ULiveLinkFaceDevice::SetHardwareId_Implementation(const FString& HardwareID)
+{
+	return false;
+}
+
+bool ULiveLinkFaceDevice::Connect_Implementation()
+{
+	const ULiveLinkFaceDeviceSettings* DeviceSettings = GetSettings();
+
+	using namespace UE::CaptureManager;
+
+	Device = FCPSDevice::MakeCPSDevice(DeviceSettings->IpAddress.IpAddressString, DeviceSettings->Port);
+	
+	Device->SubscribeToEvent(FConnectionStateChangedEvent::Name, FCaptureEventHandler(
+		FCaptureEventHandler::Type::CreateUObject(this, &ULiveLinkFaceDevice::HandleConnectionChanged),
+		EDelegateExecutionThread::AnyThread));
+
+	Device->SubscribeToEvent(FCPSStateEvent::Name, FCaptureEventHandler(
+		FCaptureEventHandler::Type::CreateUObject(this, &ULiveLinkFaceDevice::HandleCPSStateUpdate),
+		EDelegateExecutionThread::AnyThread));
+
+	Device->SubscribeToEvent(FCPSEvent::Name, FCaptureEventHandler(
+		FCaptureEventHandler::Type::CreateUObject(this, &ULiveLinkFaceDevice::HandleCPSEvent),
+		EDelegateExecutionThread::AnyThread));
+
+	Device->InitiateConnect();
+
+	return true;
+}
+
+bool ULiveLinkFaceDevice::Disconnect_Implementation()
+{
+	AsyncTask(ENamedThreads::AnyThread, [this]()
+	{
+		if (Device)
+		{
+			Device->Stop();
+			Device->UnsubscribeAll();
+			bIsConnecting = false;
+			Device = nullptr;
+		}
+	});
+	
+	return true;
+}
+
+bool ULiveLinkFaceDevice::StartRecording_Implementation()
+{
+	if (!Device)
+	{
+		return false;
+	}
+
+	using namespace UE::CaptureManager;
+
+	ILiveLinkRecordingSessionInfo& SessionInfo = ILiveLinkRecordingSessionInfo::Get();
+	if (SessionInfo.GetSlateName().IsEmpty() || SessionInfo.GetTakeNumber() == -1)
+	{
+		return false;
+	}
+
+	uint16 TakeNumber = static_cast<uint16>(SessionInfo.GetTakeNumber());
+	TProtocolResult<void> Result = Device->StartRecording(SessionInfo.GetSlateName(), TakeNumber);
+
+	return Result.HasValue();
+}
+
+bool ULiveLinkFaceDevice::StopRecording_Implementation()
+{
+	if (!Device)
+	{
+		return false;
+	}
+
+	using namespace UE::CaptureManager;
+
+	TProtocolResult<void> Result = Device->StopRecording();
+
+	return Result.HasValue();
+}
+
+bool ULiveLinkFaceDevice::IsRecording_Implementation() const
+{
+	return bIsRecording;
+}
+
+void ULiveLinkFaceDevice::HandleConnectionChanged(TSharedPtr<const UE::CaptureManager::FCaptureEvent> InEvent)
+{
+	using namespace UE::CaptureManager;
+
+	TSharedPtr<const FConnectionStateChangedEvent> Event = StaticCastSharedPtr<const FConnectionStateChangedEvent>(InEvent);
+
+	if (Event->ConnectionState == FConnectionStateChangedEvent::EState::Connecting)
+	{
+		bIsConnecting = true;
+	}
+	else if (Event->ConnectionState == FConnectionStateChangedEvent::EState::Connected)
+	{
+		bIsConnecting = false;
+	}
+
+	auto GetConnectionState = [](FConnectionStateChangedEvent::EState InConnectionState) -> ELiveLinkDeviceConnectionStatus
+	{
+		switch (InConnectionState)
+		{
+			case FConnectionStateChangedEvent::EState::Connecting:
+				return ELiveLinkDeviceConnectionStatus::Connecting;
+			case FConnectionStateChangedEvent::EState::Connected:
+				return ELiveLinkDeviceConnectionStatus::Connected;
+			case FConnectionStateChangedEvent::EState::Disconnected:
+			case FConnectionStateChangedEvent::EState::Unknown:
+			default:
+				return ELiveLinkDeviceConnectionStatus::Disconnected;
+		}
+	};
+
+	ELiveLinkDeviceConnectionStatus ConnectionStatus = GetConnectionState(Event->ConnectionState);
+
+	SetConnectionStatus(ConnectionStatus);
+}
+
+void ULiveLinkFaceDevice::HandleCPSStateUpdate(TSharedPtr<const UE::CaptureManager::FCaptureEvent> InEvent)
+{
+	using namespace UE::CaptureManager;
+
+	TSharedPtr<const FCPSStateEvent> Event = StaticCastSharedPtr<const FCPSStateEvent>(InEvent);
+
+	bIsRecording = Event->GetStateResponse.IsRecording();
+}
+
+void ULiveLinkFaceDevice::HandleCPSEvent(TSharedPtr<const UE::CaptureManager::FCaptureEvent> InEvent)
+{
+	using namespace UE::CaptureManager;
+
+	TSharedPtr<const FCPSEvent> Event = StaticCastSharedPtr<const FCPSEvent>(InEvent);
+
+	if (Event->UpdateMessage->GetAddressPath() == CPS::AddressPaths::GRecordingStatus)
+	{
+		TSharedPtr<FRecordingStatusUpdate> Update = StaticCastSharedPtr<FRecordingStatusUpdate>(Event->UpdateMessage);
+		bIsRecording = Update->IsRecording();
+	}
+	else if (Event->UpdateMessage->GetAddressPath() == CPS::AddressPaths::GTakeAdded)
+	{
+		TSharedPtr<const FTakeAddedUpdate> TakeAddedUpdate = StaticCastSharedPtr<const FTakeAddedUpdate>(Event->UpdateMessage);
+
+		TProtocolResult<FGetTakeMetadataResponse::FTakeObject> FetchTakeResult = Device->FetchTake(TakeAddedUpdate->GetTakeName());
+
+		if (FetchTakeResult.IsValid())
+		{
+			FGetTakeMetadataResponse::FTakeObject Take = FetchTakeResult.StealValue();
+			
+			FString TakeName = Take.Name;
+
+			FTakeId TakeId = AddTake(ParseTakeMetadata(Take));
+			Device->AddTakeMetadata(TakeId, MoveTemp(Take));
+
+			PublishEvent<FIngestCapability_TakeAddedEvent>(TakeId);
+
+			FetchThumbnail(TakeName, TakeId);
+		}
+	}
+	else if (Event->UpdateMessage->GetAddressPath() == CPS::AddressPaths::GTakeRemoved)
+	{
+		TSharedPtr<const FTakeRemovedUpdate> TakeRemoveUpdate = StaticCastSharedPtr<const FTakeRemovedUpdate>(Event->UpdateMessage);
+
+		FTakeId TakeId = Device->GetTakeId(TakeRemoveUpdate->GetTakeName());
+
+		if (TakeId != INDEX_NONE)
+		{
+			CancelIngest(TakeId);
+
+			RemoveTake(TakeId);
+			Device->RemoveTakeMetadata(TakeId);
+			
+			PublishEvent<FIngestCapability_TakeRemovedEvent>(TakeId);
+		}
+	}
+}
+
+FTakeMetadata ULiveLinkFaceDevice::ParseTakeMetadata(const UE::CaptureManager::FGetTakeMetadataResponse::FTakeObject& InTake)
+{
+	FTakeMetadata TakeMetadata;
+
+	TakeMetadata.Slate = InTake.Slate;
+	TakeMetadata.TakeNumber = InTake.TakeNumber;
+	FDateTime DateTime = TakeMetadata.DateTime.Get(FDateTime());
+	FDateTime::ParseIso8601(*InTake.DateTime, DateTime);
+	TakeMetadata.DateTime = DateTime;
+
+	FTakeMetadata::FVideo Video;
+	Video.FrameRate = InTake.Video.FrameRate;
+	Video.FramesCount = static_cast<uint32>(InTake.Video.Frames);
+	Video.Format = TEXT("mov");
+	Video.FrameHeight = InTake.Video.Height;
+	Video.FrameWidth = InTake.Video.Width;
+
+	TakeMetadata.Video.Add(MoveTemp(Video));
+
+	return TakeMetadata;
+}
+
+void ULiveLinkFaceDevice::OnExportProgressReport(float InProgress,
+													   UE::CaptureManager::FTaskProgress::FTask InTask)
+{
+	InTask.Update(static_cast<double>(InProgress));
+}
+
+void ULiveLinkFaceDevice::OnExportFinished(UE::CaptureManager::TProtocolResult<void> InResult,
+												 FString InTakeName,
+												 UE::CaptureManager::FTakeId InTakeId,
+												 TSharedPtr<UE::CaptureManager::FTaskProgress> InTaskProgress,
+												 TStrongObjectPtr<const UIngestCapability_Options> InIngestOptions)
+{
+	using namespace UE::CaptureManager;
+
+	const FString DownloadedStorage = InIngestOptions->DownloadDirectory;
+	const FString DownloadedTake = FPaths::Combine(DownloadedStorage, InTakeName);
+
+	if (InResult.HasValue())
+	{
+		// Blocking function
+		TOptional<FTakeMetadata> TakeMetadata = ParseTake(DownloadedStorage, InTakeName);
+		if (TakeMetadata.IsSet())
+		{
+			FScopeLock Lock(&DownloadedTakesMutex);
+			DownloadedTakes.Add(InTakeId, DownloadedTake);
+			Lock.Unlock();
+
+			UpdateTake(InTakeId, MoveTemp(TakeMetadata.GetValue()));
+			Super::IngestTake(InTakeId, InIngestOptions.Get(), InTaskProgress);
+		}
+		else
+		{
+			TValueOrError<void, FIngestCapability_Error> Result = MakeError(FIngestCapability_Error::DownloaderError, TEXT("Failed to parse the take metadata"));
+			ExecuteProcessFinishedReporter(InTakeId, MoveTemp(Result));
+		}
+	}
+	else
+	{
+		TValueOrError<void, FIngestCapability_Error> Result = MakeError(FIngestCapability_Error::DownloaderError, InResult.GetError().GetMessage());
+
+		ExecuteProcessFinishedReporter(InTakeId, MoveTemp(Result));
+	}
+
+	IFileManager::Get().DeleteDirectory(*DownloadedTake, false, true);
+}
+
+TOptional<FTakeMetadata> ULiveLinkFaceDevice::ParseTake(const FString& InTakeDirectory, const FString& InTakeName)
+{
+	using namespace UE::CaptureManager;
+
+	const FString TakeDirectory = InTakeDirectory;
+	const FString TakePath = FPaths::Combine(TakeDirectory, InTakeName);
+
+	TArray<FString> TakeFiles;
+	IFileManager::Get().FindFiles(TakeFiles, *TakePath, *FTakeMetadata::FileExtension);
+	if (!TakeFiles.IsEmpty())
+	{
+		check(TakeFiles.Num() == 1);
+
+		FTakeMetadataParser TakeMetadataParser;
+		TValueOrError<FTakeMetadata, FTakeMetadataParserError> TakeMetadataResult = TakeMetadataParser.Parse(TakePath / TakeFiles[0]);
+		if (!TakeMetadataResult.HasError())
+		{
+			return TakeMetadataResult.StealValue();
+		}
+	}
+
+	TArray<FText> ValidationFailures;
+	TOptional<FTakeMetadata> ParseOldMetadataResult = LiveLinkMetadata::ParseOldLiveLinkTakeMetadata(TakePath, ValidationFailures);
+
+	if (ParseOldMetadataResult.IsSet())
+	{
+		return ParseOldMetadataResult.GetValue();
+	}
+
+	ParseOldMetadataResult = StereoCameraMetadata::ParseOldStereoCameraMetadata(TakePath, ValidationFailures);
+
+	if (ParseOldMetadataResult.IsSet())
+	{
+		return ParseOldMetadataResult.GetValue();
+	}
+
+	return {};
+}
+
+void ULiveLinkFaceDevice::FetchThumbnail(const FString& InTakeName, UE::CaptureManager::FTakeId InTakeId)
+{
+	TMap<FString, UE::CaptureManager::FTakeId> NameToIdMap;
+	NameToIdMap.Add(InTakeName, InTakeId);
+
+	FetchThumbnails(MoveTemp(NameToIdMap));
+}
+
+void ULiveLinkFaceDevice::FetchThumbnails(TMap<FString, UE::CaptureManager::FTakeId> InNameToIdMap)
+{
+	using namespace UE::CaptureManager;
+
+	FCPSDataStream::FFileExportFinished OnFileExportFinishedCallback = FCPSDataStream::FFileExportFinished::CreateLambda(
+		[this, NameToIdMap = MoveTemp(InNameToIdMap)](const FString& InTakeName, TProtocolResult<FCPSDataStream::FData> InData)
+		{
+			if (InData.HasError())
+			{
+				return;
+			}
+
+			FCPSDataStream::FData Data = InData.StealValue();
+
+			if (const FTakeId* TakeId = NameToIdMap.Find(InTakeName))
+			{
+				TOptional<FTakeMetadata> TakeMetadataOpt = GetTakeMetadata(*TakeId);
+				if (!TakeMetadataOpt.IsSet())
+				{
+					return;
+				}
+
+				FTakeMetadata TakeMetadata = TakeMetadataOpt.GetValue();
+				TakeMetadata.Thumbnail = FTakeThumbnailData(MoveTemp(Data));
+
+				UpdateTake(*TakeId, MoveTemp(TakeMetadata));
+
+				PublishEvent<FIngestCapability_TakeUpdatedEvent>(*TakeId);
+			}
+		});
+
+	TUniquePtr<FCPSDataStream> DataStream = MakeUnique<FCPSDataStream>(MoveTemp(OnFileExportFinishedCallback));
+
+	Device->FetchThumbnails(MoveTemp(DataStream));
+}
diff --git a/Engine/Plugins/VirtualProduction/CaptureManager/CaptureManagerApp/Source/GenericLiveLinkDevice/Private/CPS/LiveLinkFaceDevice.h b/Engine/Plugins/VirtualProduction/CaptureManager/CaptureManagerApp/Source/GenericLiveLinkDevice/Private/CPS/LiveLinkFaceDevice.h
new file mode 100644
--- /dev/null
+++ b/Engine/Plugins/VirtualProduction/CaptureManager/CaptureManagerApp/Source/GenericLiveLinkDevice/Private/CPS/LiveLinkFaceDevice.h
@@ -0,0 +1,112 @@
+// Copyright Epic Games, Inc. All Rights Reserved.
+
+#pragma once
+
+#include "LiveLinkDevice.h"
+
+#include "Ingest/LiveLinkDeviceCapability_Ingest.h"
+#include "LiveLinkDeviceCapability_Connection.h"
+#include "LiveLinkDeviceCapability_Recording.h"
+#include "Misc/Optional.h"
+
+#include "Async/TaskProgress.h"
+
+#include "Protocol/CPSDevice.h"
+#include "Customizations/ToggleConnectActionCustomization.h"
+#include "Customizations/DeviceIpAddressCustomization.h"
+
+#include "BaseIngestLiveLinkDevice.h"
+
+#include "LiveLinkFaceDevice.generated.h"
+
+
+UCLASS(BlueprintType)
+class ULiveLinkFaceDeviceSettings : public ULiveLinkDeviceSettings
+{
+	GENERATED_BODY()
+
+public:
+	UPROPERTY(EditAnywhere, Category = "Live Link Face Device")
+	FString DisplayName = TEXT("Live Link Face Device");
+
+	UPROPERTY(EditAnywhere, Category = "Live Link Face Device")
+	FDeviceIpAddress IpAddress;
+
+	UPROPERTY(EditAnywhere, Category = "Live Link Face Device")
+	uint16 Port = 14785;
+
+	UPROPERTY(VisibleAnywhere, Category = "Live Link Face Device")
+	FToggleConnectAction ConnectAction;
+};
+
+
+UCLASS(BlueprintType)
+class ULiveLinkFaceDevice final : public UBaseIngestLiveLinkDevice
+	, public ILiveLinkDeviceCapability_Connection
+	, public ILiveLinkDeviceCapability_Recording
+{
+	GENERATED_BODY()
+
+public:
+	UFUNCTION(BlueprintCallable, Category = "LiveLinkFaceDeviceSettings")
+	const ULiveLinkFaceDeviceSettings* GetSettings() const;
+
+	//~ Begin ULiveLinkDevice interface
+	virtual TSubclassOf<ULiveLinkDeviceSettings> StaticGetSettingsClass() const override;
+	virtual FText GetDisplayName() const override;
+	virtual EDeviceHealth GetDeviceHealth() const override;
+	virtual FText GetHealthText() const override;
+
+	virtual void OnDeviceAdded(const FGuid InDeviceGuid, ULiveLinkDeviceSettings* InSettings) override;
+	virtual void OnDeviceRemoved() override;
+	//~ End ULiveLinkDevice interface
+
+private:
+	virtual FString GetFullTakePath(UE::CaptureManager::FTakeId InTakeId) const override;
+
+	//~ Begin ULiveLinkDeviceCapability_Ingest interface
+	virtual void UpdateTakeList_Implementation() override;
+	virtual void RunIngestTake(int32 InTakeId, const UIngestCapability_Options* InIngestOptions) override;
+	virtual void CancelIngest(int32 InTakeId) override;
+	//~ End ULiveLinkDeviceCapability_Ingest interface
+
+	//~ Begin ILiveLinkDeviceCapability_Connection interface
+	virtual ELiveLinkDeviceConnectionStatus GetConnectionStatus_Implementation() const override;
+	virtual FString GetHardwareId_Implementation() const override;
+	virtual bool SetHardwareId_Implementation(const FString& HardwareID) override;
+	virtual bool Connect_Implementation() override;
+	virtual bool Disconnect_Implementation() override;
+	//~ End ILiveLinkDeviceCapability_Connection interface
+
+	//~ Begin ILiveLinkDeviceCapability_Recording interface
+	virtual bool StartRecording_Implementation() override;
+	virtual bool StopRecording_Implementation() override;
+	virtual bool IsRecording_Implementation() const override;
+	//~ End ILiveLinkDeviceCapability_Recording interface
+
+	void HandleConnectionChanged(TSharedPtr<const UE::CaptureManager::FCaptureEvent> InEvent);
+	void HandleCPSStateUpdate(TSharedPtr<const UE::CaptureManager::FCaptureEvent> InEvent);
+	void HandleCPSEvent(TSharedPtr<const UE::CaptureManager::FCaptureEvent> InEvent);
+
+	void OnExportProgressReport(float InProgress,
+								UE::CaptureManager::FTaskProgress::FTask InTask);
+	void OnExportFinished(UE::CaptureManager::TProtocolResult<void> InResult,
+						  FString InTakeName,
+						  UE::CaptureManager::FTakeId InTakeId,
+						  TSharedPtr<UE::CaptureManager::FTaskProgress> InTaskProgress,
+						  TStrongObjectPtr<const UIngestCapability_Options> InIngestOptions);
+
+	TOptional<FTakeMetadata> ParseTake(const FString& InTakeDirectory, const FString& InTakeName);
+	FTakeMetadata ParseTakeMetadata(const UE::CaptureManager::FGetTakeMetadataResponse::FTakeObject& InTake);
+
+	void FetchThumbnail(const FString& InTakeName, UE::CaptureManager::FTakeId InTakeId);
+	void FetchThumbnails(TMap<FString, UE::CaptureManager::FTakeId> InNameToIdMap);
+
+	TSharedPtr<UE::CaptureManager::FCPSDevice> Device;
+
+	mutable FCriticalSection DownloadedTakesMutex;
+	TMap<UE::CaptureManager::FTakeId, FString> DownloadedTakes;
+
+	std::atomic_bool bIsRecording = false;
+	std::atomic_bool bIsConnecting = false;
+};
diff --git a/Engine/Plugins/VirtualProduction/CaptureManager/CaptureManagerApp/Source/GenericLiveLinkDevice/Private/Filesystem/FilesystemLiveLinkDevice.cpp b/Engine/Plugins/VirtualProduction/CaptureManager/CaptureManagerApp/Source/GenericLiveLinkDevice/Private/Filesystem/FilesystemLiveLinkDevice.cpp
new file mode 100644
--- /dev/null
+++ b/Engine/Plugins/VirtualProduction/CaptureManager/CaptureManagerApp/Source/GenericLiveLinkDevice/Private/Filesystem/FilesystemLiveLinkDevice.cpp
@@ -0,0 +1,284 @@
+// Copyright Epic Games, Inc. All Rights Reserved.
+
+#include "FilesystemLiveLinkDevice.h"
+
+#include "HAL/FileManager.h"
+#include "Settings/CaptureManagerSettings.h"
+
+#include "Utils/GenericLiveLinkDeviceLog.h"
+
+#include "LiveLinkFaceMetadata.h"
+#include "StereoCameraMetadataParseUtils.h"
+
+#include "Utils/GenericLiveLinkDeviceUtility.h"
+
+DEFINE_LOG_CATEGORY(LogGenericLiveLinkArchiveDevice);
+
+const UFilesystemLiveLinkDeviceSettings* UFilesystemLiveLinkDevice::GetSettings() const
+{
+	return CastChecked<UFilesystemLiveLinkDeviceSettings>(Settings);
+}
+
+TSubclassOf<ULiveLinkDeviceSettings> UFilesystemLiveLinkDevice::StaticGetSettingsClass() const
+{
+	return UFilesystemLiveLinkDeviceSettings::StaticClass();
+}
+
+FText UFilesystemLiveLinkDevice::GetDisplayName() const
+{
+	return FText::FromString(GetSettings()->DisplayName);
+}
+
+EDeviceHealth UFilesystemLiveLinkDevice::GetDeviceHealth() const
+{
+	return EDeviceHealth::Nominal;
+}
+
+FText UFilesystemLiveLinkDevice::GetHealthText() const
+{
+	return FText::FromString("Nominal");
+}
+
+void UFilesystemLiveLinkDevice::OnSettingChanged(const FPropertyChangedEvent& InPropertyChangedEvent)
+{
+	Super::OnSettingChanged(InPropertyChangedEvent);
+
+	if (GET_MEMBER_NAME_CHECKED(UFilesystemLiveLinkDeviceSettings, StoragePath) == InPropertyChangedEvent.GetMemberPropertyName())
+	{
+		FString Path = GetSettings()->StoragePath.Path;
+		if (!Path.IsEmpty() && FPaths::DirectoryExists(Path))
+		{
+			SetConnectionStatus(ELiveLinkDeviceConnectionStatus::Connected);
+		}
+		else
+		{
+			SetConnectionStatus(ELiveLinkDeviceConnectionStatus::Disconnected);
+		}
+	}
+}
+
+FString UFilesystemLiveLinkDevice::GetFullTakePath(UE::CaptureManager::FTakeId InTakeId) const
+{
+	FString StoragePath = GetSettings()->StoragePath.Path;
+
+	if (const FString* RelativeTakePath = RelativeTakePaths.Find(InTakeId))
+	{
+		return FPaths::Combine(StoragePath, *RelativeTakePath);
+	}
+
+	return FString();
+}
+
+void UFilesystemLiveLinkDevice::RunIngestTake(int32 InTakeId, const UIngestCapability_Options* InIngestOptions)
+{
+	AsyncTask(ENamedThreads::AnyThread, [this,
+			  TakeId = InTakeId,
+			  IngestOptions = InIngestOptions]()
+	{
+		static constexpr uint32 NumberOfTasks = 2; // Convert, Upload
+
+		using namespace UE::CaptureManager;
+		TSharedPtr<FTaskProgress> TaskProgress
+			= MakeShared<FTaskProgress>(NumberOfTasks, FTaskProgress::FProgressReporter::CreateLambda([this, TakeId](double InProgress)
+		{
+			ExecuteProcessProgressReporter(TakeId, InProgress);
+		}));
+
+		Super::IngestTake(TakeId, IngestOptions, MoveTemp(TaskProgress));
+	});
+}
+
+void UFilesystemLiveLinkDevice::UpdateTakeList_Implementation()
+{
+	AsyncTask(ENamedThreads::AnyThread, [this]()
+	{
+		RemoveAllTakes();
+
+		FString StoragePath = GetSettings()->StoragePath.Path;
+
+		TArray<FString> TakeMetadataFiles;
+		const bool bIterationResult = IFileManager::Get().IterateDirectoryRecursively(*StoragePath, [&](const TCHAR* InFileNameOrDirectory, bool bInIsDirectory)
+		{
+			if (!bInIsDirectory)
+			{
+				const FString CurrentFileName = FPaths::GetCleanFilename(InFileNameOrDirectory);
+				if (CurrentFileName == TEXT("take.json"))
+				{
+					TakeMetadataFiles.Add(InFileNameOrDirectory);
+				}
+				else if (FPaths::GetExtension(CurrentFileName) == FTakeMetadata::FileExtension)
+				{
+					TakeMetadataFiles.Add(InFileNameOrDirectory);
+				}
+			}
+
+			return true;
+		});
+
+		for (const FString& CurrentTakeFile : TakeMetadataFiles)
+		{
+			TOptional<FTakeMetadata> TakeInfoResult = ReadTake(CurrentTakeFile);
+
+			if (TakeInfoResult.IsSet())
+			{
+				int32 TakeId = AddTake(TakeInfoResult.GetValue());
+				FString CurrentDirectory = FPaths::GetPath(CurrentTakeFile);
+
+				FString TakePath = CurrentDirectory.Right(CurrentDirectory.Len() - StoragePath.Len());
+				FPaths::NormalizeDirectoryName(TakePath); // Removes trailing slash
+				RelativeTakePaths.Add(TakeId, MoveTemp(TakePath));
+			}
+		}
+
+		ExecuteUpdateTakeListCallback(Execute_GetTakeIdentifiers(this));
+	});
+}
+
+ELiveLinkDeviceConnectionStatus UFilesystemLiveLinkDevice::GetConnectionStatus_Implementation() const
+{
+	const UFilesystemLiveLinkDeviceSettings* DeviceSettings = GetSettings();
+	FString Path = DeviceSettings->StoragePath.Path;
+
+	if (!Path.IsEmpty() && FPaths::DirectoryExists(Path))
+	{
+		return ELiveLinkDeviceConnectionStatus::Connected;
+	}
+
+	return ELiveLinkDeviceConnectionStatus::Disconnected;
+}
+
+FString UFilesystemLiveLinkDevice::GetHardwareId_Implementation() const
+{
+	return FPlatformMisc::GetDeviceId();
+}
+
+bool UFilesystemLiveLinkDevice::SetHardwareId_Implementation(const FString& HardwareID)
+{
+	return false;
+}
+
+bool UFilesystemLiveLinkDevice::Connect_Implementation()
+{
+	const UFilesystemLiveLinkDeviceSettings* DeviceSettings = GetSettings();
+	FString Path = DeviceSettings->StoragePath.Path;
+
+	if (!Path.IsEmpty() && FPaths::DirectoryExists(Path))
+	{
+		SetConnectionStatus(ELiveLinkDeviceConnectionStatus::Connected);
+		return true;
+	}
+	
+	return false;
+}
+
+bool UFilesystemLiveLinkDevice::Disconnect_Implementation()
+{
+	SetConnectionStatus(ELiveLinkDeviceConnectionStatus::Disconnected);
+	return true;
+}
+
+TOptional<FTakeMetadata> UFilesystemLiveLinkDevice::ReadTake(const FString& InCurrentTakeFile)
+{
+	using namespace UE::CaptureManager;
+
+	FTakeMetadataParser TakeMetadataParser;
+	TValueOrError<FTakeMetadata, FTakeMetadataParserError> TakeMetadataResult = TakeMetadataParser.Parse(InCurrentTakeFile);
+	FString CurrentDirectory = FPaths::GetPath(InCurrentTakeFile);
+
+	if (TakeMetadataResult.HasError())
+	{
+		FTakeMetadataParserError ParserError = TakeMetadataResult.StealError();
+		FString OriginName = UE::CaptureManager::Private::ErrorOriginToString(ParserError.Origin);
+
+		FString Message =
+			FString::Format(TEXT("Unable to parse take metadata file - {0} (Error origin: {1}): {2}"),
+							{
+								InCurrentTakeFile,
+								ParserError.Message.ToString(),
+								OriginName
+							});
+
+		UE_LOG(LogGenericLiveLinkArchiveDevice, Warning, TEXT("%s"), *Message);
+
+		UE_LOG(LogGenericLiveLinkArchiveDevice, Display, TEXT("Checking backwards compatible take metadata formats"), *CurrentDirectory);
+
+		TArray<FText> ValidationFailures;
+
+		UE_LOG(LogGenericLiveLinkArchiveDevice, Display, TEXT("Checking directory (%s) for pre UE 5.6 LiveLink take metadata format"), *CurrentDirectory);
+		TOptional<FTakeMetadata> ParseOldMetadataResult = LiveLinkMetadata::ParseOldLiveLinkTakeMetadata(CurrentDirectory, ValidationFailures);
+
+		if (ParseOldMetadataResult.IsSet())
+		{
+			return ParseOldMetadataResult.GetValue();
+		}
+
+		UE_LOG(LogGenericLiveLinkArchiveDevice, Display, TEXT("Checking directory (%s) for pre UE 5.6 StereoCamera take metadata format"), *CurrentDirectory);
+		ParseOldMetadataResult = StereoCameraMetadata::ParseOldStereoCameraMetadata(CurrentDirectory, ValidationFailures);
+
+		if (ParseOldMetadataResult.IsSet())
+		{
+			return ParseOldMetadataResult.GetValue();
+		}
+
+		UE_LOG(LogGenericLiveLinkArchiveDevice, Error, TEXT("Failed to parse take metadata file: %s"), *CurrentDirectory);
+
+		return {};
+	}
+
+	FTakeMetadata& TakeMetadata = TakeMetadataResult.GetValue();
+
+	TOptional<FString> ThumbnailPath = TakeMetadata.Thumbnail.GetThumbnailPath();
+	if (ThumbnailPath.IsSet())
+	{
+		if (FPaths::IsRelative(*ThumbnailPath))
+		{
+			TakeMetadata.Thumbnail = FTakeThumbnailData(FPaths::ConvertRelativePathToFull(CurrentDirectory, *ThumbnailPath));
+		}
+	}
+
+	if (!TakeMetadata.Video.IsEmpty())
+	{
+		for (FTakeMetadata::FVideo& Video : TakeMetadata.Video)
+		{
+			if (FPaths::IsRelative(*Video.Path))
+			{
+				Video.Path = FPaths::ConvertRelativePathToFull(CurrentDirectory, *Video.Path);
+			}
+		}
+	}
+
+	if (!TakeMetadata.Depth.IsEmpty())
+	{
+		for (FTakeMetadata::FVideo& Depth : TakeMetadata.Depth)
+		{
+			if (FPaths::IsRelative(*Depth.Path))
+			{
+				Depth.Path = FPaths::ConvertRelativePathToFull(CurrentDirectory, *Depth.Path);
+			}
+		}
+	}
+
+	if (!TakeMetadata.Audio.IsEmpty())
+	{
+		for (FTakeMetadata::FAudio& Audio : TakeMetadata.Audio)
+		{
+			if (FPaths::IsRelative(*Audio.Path))
+			{
+				Audio.Path = FPaths::ConvertRelativePathToFull(CurrentDirectory, *Audio.Path);
+			}
+		}
+	}
+
+	if (!TakeMetadata.Calibration.IsEmpty())
+	{
+		for (FTakeMetadata::FCalibration& Calib : TakeMetadata.Calibration)
+		{
+			if (FPaths::IsRelative(*Calib.Path))
+			{
+				Calib.Path = FPaths::ConvertRelativePathToFull(CurrentDirectory, *Calib.Path);
+			}
+		}
+	}
+
+	return TakeMetadataResult.StealValue();
+}
\ No newline at end of file
diff --git a/Engine/Plugins/VirtualProduction/CaptureManager/CaptureManagerApp/Source/GenericLiveLinkDevice/Private/Filesystem/FilesystemLiveLinkDevice.h b/Engine/Plugins/VirtualProduction/CaptureManager/CaptureManagerApp/Source/GenericLiveLinkDevice/Private/Filesystem/FilesystemLiveLinkDevice.h
new file mode 100644
--- /dev/null
+++ b/Engine/Plugins/VirtualProduction/CaptureManager/CaptureManagerApp/Source/GenericLiveLinkDevice/Private/Filesystem/FilesystemLiveLinkDevice.h
@@ -0,0 +1,71 @@
+// Copyright Epic Games, Inc. All Rights Reserved.
+
+#pragma once
+
+#include "LiveLinkDevice.h"
+
+#include "Ingest/LiveLinkDeviceCapability_Ingest.h"
+#include "LiveLinkDeviceCapability_Connection.h"
+
+#include "Misc/Optional.h"
+
+#include "Async/TaskProgress.h"
+
+#include "BaseIngestLiveLinkDevice.h"
+
+#include "FilesystemLiveLinkDevice.generated.h"
+
+
+UCLASS(BlueprintType)
+class UFilesystemLiveLinkDeviceSettings : public ULiveLinkDeviceSettings
+{
+	GENERATED_BODY()
+
+public:
+	UPROPERTY(EditAnywhere, Category = "Filesystem Device")
+	FString DisplayName = TEXT("Filesystem Device");
+
+	UPROPERTY(EditAnywhere, Category = "Filesystem Device")
+	FDirectoryPath StoragePath;
+};
+
+
+UCLASS(BlueprintType)
+class UFilesystemLiveLinkDevice final : public UBaseIngestLiveLinkDevice
+	, public ILiveLinkDeviceCapability_Connection
+{
+	GENERATED_BODY()
+
+public:
+	UFUNCTION(BlueprintCallable, Category = "FilesystemDeviceSettings")
+	const UFilesystemLiveLinkDeviceSettings* GetSettings() const;
+
+	//~ Begin ULiveLinkDevice interface
+	virtual TSubclassOf<ULiveLinkDeviceSettings> StaticGetSettingsClass() const override;
+	virtual FText GetDisplayName() const override;
+	virtual EDeviceHealth GetDeviceHealth() const override;
+	virtual FText GetHealthText() const override;
+
+	virtual void OnSettingChanged(const FPropertyChangedEvent& InPropertyChangedEvent) override;
+	//~ End ULiveLinkDevice interface
+
+private:
+	virtual FString GetFullTakePath(UE::CaptureManager::FTakeId InTakeId) const override;
+
+	//~ Begin ULiveLinkDeviceCapability_Ingest interface
+	virtual void RunIngestTake(int32 InTakeId, const UIngestCapability_Options* InIngestOptions) override;
+	virtual void UpdateTakeList_Implementation() override;
+	//~ End ULiveLinkDeviceCapability_Ingest interface
+
+	//~ Begin ILiveLinkDeviceCapability_Connection interface
+	virtual ELiveLinkDeviceConnectionStatus GetConnectionStatus_Implementation() const override;
+	virtual FString GetHardwareId_Implementation() const override;
+	virtual bool SetHardwareId_Implementation(const FString& HardwareID) override;
+	virtual bool Connect_Implementation() override;
+	virtual bool Disconnect_Implementation() override;
+	//~ End ILiveLinkDeviceCapability_Connection interface
+
+	static TOptional<FTakeMetadata> ReadTake(const FString& InCurrentTakeFile);
+
+	TMap<int32, FString> RelativeTakePaths;
+};
diff --git a/Engine/Plugins/VirtualProduction/CaptureManager/CaptureManagerApp/Source/GenericLiveLinkDevice/Private/GenericLiveLinkDeviceHelper.cpp b/Engine/Plugins/VirtualProduction/CaptureManager/CaptureManagerApp/Source/GenericLiveLinkDevice/Private/GenericLiveLinkDeviceHelper.cpp
--- a/Engine/Plugins/VirtualProduction/CaptureManager/CaptureManagerApp/Source/GenericLiveLinkDevice/Private/GenericLiveLinkDeviceHelper.cpp
+++ b/Engine/Plugins/VirtualProduction/CaptureManager/CaptureManagerApp/Source/GenericLiveLinkDevice/Private/GenericLiveLinkDeviceHelper.cpp
@@ -8,12 +8,12 @@
 
 #include "Ingest/LiveLinkDeviceCapability_Ingest.h"
 
-#include "Archive/GenericLiveLinkArchiveDevice.h"
-#include "CPS/GenericLiveLinkCPSDevice.h"
+#include "Filesystem/FilesystemLiveLinkDevice.h"
+#include "CPS/LiveLinkFaceDevice.h"
 
 ULiveLinkDevice* UGenericLiveLinkDeviceHelper::CreateArchiveDevice(FString InName, FString InStoragePath) const
 {
-	TSubclassOf<ULiveLinkDevice> DeviceClass = UGenericLiveLinkArchiveDevice::StaticClass();
+	TSubclassOf<ULiveLinkDevice> DeviceClass = UFilesystemLiveLinkDevice::StaticClass();
 
 	ULiveLinkDeviceSubsystem* DeviceSubsystem = GEngine->GetEngineSubsystem<ULiveLinkDeviceSubsystem>();
 	ULiveLinkDevice* NewDevice = NewObject<ULiveLinkDevice>(DeviceSubsystem, DeviceClass);
@@ -25,7 +25,7 @@ ULiveLinkDevice* UGenericLiveLinkDeviceHelper::CreateArchiveDevice(FString InNam
 	TSubclassOf<ULiveLinkDeviceSettings> SettingsClass = NewDevice->StaticGetSettingsClass();
 	ULiveLinkDeviceSettings* NewSettings = NewObject<ULiveLinkDeviceSettings>(NewDevice, SettingsClass);
 
-	UGenericLiveLinkArchiveDeviceSettings* DeviceSettings = CastChecked<UGenericLiveLinkArchiveDeviceSettings>(NewSettings);
+	UFilesystemLiveLinkDeviceSettings* DeviceSettings = CastChecked<UFilesystemLiveLinkDeviceSettings>(NewSettings);
 	DeviceSettings->StoragePath.Path = MoveTemp(InStoragePath);
 	DeviceSettings->DisplayName = MoveTemp(InName);
 
@@ -36,7 +36,7 @@ ULiveLinkDevice* UGenericLiveLinkDeviceHelper::CreateArchiveDevice(FString InNam
 
 ULiveLinkDevice* UGenericLiveLinkDeviceHelper::CreateNetworkDevice(FString InName, FString InIPAddress, int32 InPort) const
 {
-	TSubclassOf<ULiveLinkDevice> DeviceClass = UGenericLiveLinkCPSDevice::StaticClass();
+	TSubclassOf<ULiveLinkDevice> DeviceClass = ULiveLinkFaceDevice::StaticClass();
 
 	ULiveLinkDeviceSubsystem* DeviceSubsystem = GEngine->GetEngineSubsystem<ULiveLinkDeviceSubsystem>();
 	ULiveLinkDevice* NewDevice = NewObject<ULiveLinkDevice>(DeviceSubsystem, DeviceClass);
@@ -48,7 +48,7 @@ ULiveLinkDevice* UGenericLiveLinkDeviceHelper::CreateNetworkDevice(FString InNam
 	TSubclassOf<ULiveLinkDeviceSettings> SettingsClass = NewDevice->StaticGetSettingsClass();
 	ULiveLinkDeviceSettings* NewSettings = NewObject<ULiveLinkDeviceSettings>(NewDevice, SettingsClass);
 
-	UGenericLiveLinkCPSDeviceSettings* DeviceSettings = CastChecked<UGenericLiveLinkCPSDeviceSettings>(NewSettings);
+	ULiveLinkFaceDeviceSettings* DeviceSettings = CastChecked<ULiveLinkFaceDeviceSettings>(NewSettings);
 	DeviceSettings->IpAddress.IpAddressString = MoveTemp(InIPAddress);
 	DeviceSettings->Port = static_cast<uint16>(InPort);
 	DeviceSettings->DisplayName = MoveTemp(InName);
diff --git a/Engine/Plugins/VirtualProduction/CaptureManager/CaptureManagerApp/Source/GenericLiveLinkDevice/Private/GenericLiveLinkDeviceModule.cpp b/Engine/Plugins/VirtualProduction/CaptureManager/CaptureManagerApp/Source/GenericLiveLinkDevice/Private/GenericLiveLinkDeviceModule.cpp
--- a/Engine/Plugins/VirtualProduction/CaptureManager/CaptureManagerApp/Source/GenericLiveLinkDevice/Private/GenericLiveLinkDeviceModule.cpp
+++ b/Engine/Plugins/VirtualProduction/CaptureManager/CaptureManagerApp/Source/GenericLiveLinkDevice/Private/GenericLiveLinkDeviceModule.cpp
@@ -4,7 +4,7 @@
 
 #include "PropertyEditorModule.h"
 
-#include "CPS/GenericLiveLinkCPSDevice.h"
+#include "CPS/LiveLinkFaceDevice.h"
 #include "CPS/Customizations/ToggleConnectActionCustomization.h"
 #include "CPS/Customizations/DeviceIpAddressCustomization.h"
 
diff --git a/Engine/Plugins/VirtualProduction/CaptureManager/CaptureManagerApp/Source/GenericLiveLinkDevice/Private/Tests/GenericLiveLinkDeviceTests.cpp b/Engine/Plugins/VirtualProduction/CaptureManager/CaptureManagerApp/Source/GenericLiveLinkDevice/Private/Tests/GenericLiveLinkDeviceTests.cpp
--- a/Engine/Plugins/VirtualProduction/CaptureManager/CaptureManagerApp/Source/GenericLiveLinkDevice/Private/Tests/GenericLiveLinkDeviceTests.cpp
+++ b/Engine/Plugins/VirtualProduction/CaptureManager/CaptureManagerApp/Source/GenericLiveLinkDevice/Private/Tests/GenericLiveLinkDeviceTests.cpp
@@ -1,15 +1,15 @@
 // Copyright Epic Games, Inc. All Rights Reserved.
 
-#include "Archive/GenericLiveLinkArchiveDevice.h"
+#include "Filesystem/FilesystemLiveLinkDevice.h"
 
 #include "Misc/AutomationTest.h"
 
-IMPLEMENT_SIMPLE_AUTOMATION_TEST(FGenericLiveLinkArchiveDeviceCreationTest, "CaptureManager.GenericLiveLinkArchiveDevice.Creation",
+IMPLEMENT_SIMPLE_AUTOMATION_TEST(FGenericLiveLinkArchiveDeviceCreationTest, "CaptureManager.FilesystemLiveLinkDevice.Creation",
 								 EAutomationTestFlags::EditorContext | EAutomationTestFlags::EngineFilter);
 
 bool FGenericLiveLinkArchiveDeviceCreationTest::RunTest(const FString& Parameters)
 {
-	UGenericLiveLinkArchiveDevice* Device = NewObject<UGenericLiveLinkArchiveDevice>();
+	UFilesystemLiveLinkDevice* Device = NewObject<UFilesystemLiveLinkDevice>();
 
 	TestNotNull(TEXT("Device Not Null"), Device);
 	TestTrue(TEXT("Device Implements Ingest Capability"), Device->Implements<ULiveLinkDeviceCapability_Ingest>());
@@ -17,12 +17,12 @@ bool FGenericLiveLinkArchiveDeviceCreationTest::RunTest(const FString& Parameter
 	return true;
 }
 
-IMPLEMENT_SIMPLE_AUTOMATION_TEST(FGenericLiveLinkArchiveDeviceIngestTakeAsyncTest, "CaptureManager.GenericLiveLinkArchiveDevice.IngestTakeAsync",
+IMPLEMENT_SIMPLE_AUTOMATION_TEST(FGenericLiveLinkArchiveDeviceIngestTakeAsyncTest, "CaptureManager.FilesystemLiveLinkDevice.IngestTakeAsync",
 								 EAutomationTestFlags::EditorContext | EAutomationTestFlags::EngineFilter);
 
 bool FGenericLiveLinkArchiveDeviceIngestTakeAsyncTest::RunTest(const FString& Parameters)
 {
-	UGenericLiveLinkArchiveDevice* Device = NewObject<UGenericLiveLinkArchiveDevice>();
+	UFilesystemLiveLinkDevice* Device = NewObject<UFilesystemLiveLinkDevice>();
 
 	TestNotNull(TEXT("Device Not Null"), Device);
 	TestTrue(TEXT("Device Implements Ingest Capability"), Device->Implements<ULiveLinkDeviceCapability_Ingest>());
