[CaptureManager] Enable device tabs by default in Capture Manager layout

diff --git a/Engine/Plugins/VirtualProduction/CaptureManager/CaptureManagerApp/Source/CaptureManagerEditor/Private/UI/LiveLinkHubCaptureManagerMode.cpp b/Engine/Plugins/VirtualProduction/CaptureManager/CaptureManagerApp/Source/CaptureManagerEditor/Private/UI/LiveLinkHubCaptureManagerMode.cpp
--- a/Engine/Plugins/VirtualProduction/CaptureManager/CaptureManagerApp/Source/CaptureManagerEditor/Private/UI/LiveLinkHubCaptureManagerMode.cpp
+++ b/Engine/Plugins/VirtualProduction/CaptureManager/CaptureManagerApp/Source/CaptureManagerEditor/Private/UI/LiveLinkHubCaptureManagerMode.cpp
@@ -162,9 +162,21 @@ FLiveLinkHubCaptureManagerMode::FLiveLinkHubCaptureManagerMode(TSharedPtr<FLiveL
 				->SetOrientation(Orient_Horizontal)
 				->Split
 				(
-					FTabManager::NewStack()
-					->SetSizeCoefficient(0.5f)
-					->AddTab("LiveLinkDevices", ETabState::OpenedTab)
+					FTabManager::NewSplitter()
+					->SetSizeCoefficient(0.25f)
+					->SetOrientation(Orient_Vertical)
+					->Split
+					(
+						FTabManager::NewStack()
+						->SetSizeCoefficient(0.5f)
+						->AddTab("LiveLinkDevices", ETabState::OpenedTab)
+					)
+					->Split
+					(
+						FTabManager::NewStack()
+						->SetSizeCoefficient(0.5f)
+						->AddTab("LiveLinkDeviceDetails", ETabState::OpenedTab)
+					)
 				)
 				->Split
 				(
