[CaptureManager] Hardcoded names for live link face media tracks during ingest

diff --git a/Engine/Plugins/VirtualProduction/CaptureManager/CaptureManagerApp/Source/LiveLinkFaceMetadata/Private/LiveLinkFaceMetadata.cpp b/Engine/Plugins/VirtualProduction/CaptureManager/CaptureManagerApp/Source/LiveLinkFaceMetadata/Private/LiveLinkFaceMetadata.cpp
--- a/Engine/Plugins/VirtualProduction/CaptureManager/CaptureManagerApp/Source/LiveLinkFaceMetadata/Private/LiveLinkFaceMetadata.cpp
+++ b/Engine/Plugins/VirtualProduction/CaptureManager/CaptureManagerApp/Source/LiveLinkFaceMetadata/Private/LiveLinkFaceMetadata.cpp
@@ -923,6 +923,7 @@ FTakeMetadata ConvertOldToNewTakeMetadata(const FLiveLinkFaceTakeInfo& InTakeInf
 	Parser.ParseFrameLog(InTakeInfo.GetFrameLogFilePath(), InTakeInfo.VideoMetadata.FrameRate);
 
 	FTakeMetadata::FVideo Video;
+	Video.Name = "Video";
 	Video.FrameWidth = InTakeInfo.VideoMetadata.Resolution.X;
 	Video.FrameHeight = InTakeInfo.VideoMetadata.Resolution.Y;
 	Video.Orientation = ConvertOldOrientation(InTakeInfo.VideoMetadata.Orientation);
@@ -939,6 +940,7 @@ FTakeMetadata ConvertOldToNewTakeMetadata(const FLiveLinkFaceTakeInfo& InTakeInf
 	if (HasDepthMetadataFile(InTakeInfo.TakeOriginDirectory))
 	{
 		FTakeMetadata::FVideo Depth;
+		Depth.Name = "Depth";
 		Depth.FramesCount = Parser.DepthFrameCount;
 		Depth.FrameRate = InTakeInfo.DepthMetadata.FrameRate;
 		Depth.FrameWidth = InTakeInfo.DepthMetadata.Resolution.X;
@@ -954,6 +956,7 @@ FTakeMetadata ConvertOldToNewTakeMetadata(const FLiveLinkFaceTakeInfo& InTakeInf
 	}
 
 	FTakeMetadata::FAudio Audio;
+	Audio.Name = "Audio";
 	Audio.Path = InTakeInfo.GetVideoFilePath();
 
 	// We do not know the audio duration and so we must estimate it from the video. This apporach has been deemed 
