Take names in the Content Browser are wrong

diff --git a/Engine/Plugins/VirtualProduction/CaptureManager/CaptureManagerApp/Source/CaptureManagerEditor/Private/UI/TakesView.cpp b/Engine/Plugins/VirtualProduction/CaptureManager/CaptureManagerApp/Source/CaptureManagerEditor/Private/UI/TakesView.cpp
--- a/Engine/Plugins/VirtualProduction/CaptureManager/CaptureManagerApp/Source/CaptureManagerEditor/Private/UI/TakesView.cpp
+++ b/Engine/Plugins/VirtualProduction/CaptureManager/CaptureManagerApp/Source/CaptureManagerEditor/Private/UI/TakesView.cpp
@@ -551,10 +551,10 @@ TArray<TStrongObjectPtr<UTakeVirtualAsset>> FTakesView::MakeTakeObjects(FGuid In
 
 		FTakeMetadata TakeInfo = TakeInfoOpt.GetValue();
 
-		FName DesiredName(TakeInfo.Slate, TakeInfo.TakeNumber);
-		FName Name = MakeUniqueObjectName(Package, UTakeVirtualAsset::StaticClass(), FName(DesiredName));
+		FString DesiredNameStr = FString::Format(TEXT("{0}_{1}"), { TakeInfo.Slate, TakeInfo.TakeNumber });
+		FName DesiredName(MoveTemp(DesiredNameStr));
 
-		TStrongObjectPtr<UTakeVirtualAsset> TakeObject(NewObject<UTakeVirtualAsset>(Package, Name));
+		TStrongObjectPtr<UTakeVirtualAsset> TakeObject(NewObject<UTakeVirtualAsset>(Package, DesiredName));
 		TakeObject->CaptureDeviceId = InCaptureDeviceId;
 		TakeObject->TakeId = TakeId;
 		TakeObject->Metadata = MoveTemp(TakeInfo);
