[CaptureManager] Replaced preset selection with pipeline selection and unified location of tool buttons

diff --git a/Engine/Plugins/VirtualProduction/CaptureManager/CaptureManagerApp/Source/CaptureManagerEditor/Private/IngestManagement/IngestPipelineManager.cpp b/Engine/Plugins/VirtualProduction/CaptureManager/CaptureManagerApp/Source/CaptureManagerEditor/Private/IngestManagement/IngestPipelineManager.cpp
new file mode 100644
--- /dev/null
+++ b/Engine/Plugins/VirtualProduction/CaptureManager/CaptureManagerApp/Source/CaptureManagerEditor/Private/IngestManagement/IngestPipelineManager.cpp
@@ -0,0 +1,52 @@
+// Copyright Epic Games, Inc. All Rights Reserved.
+
+#include "IngestPipelineManager.h"
+
+#include "Internationalization/Internationalization.h"
+
+#define LOCTEXT_NAMESPACE "IngestPipelineManager"
+
+namespace UE::CaptureManager
+{
+
+FIngestPipelineManager::FIngestPipelineManager()
+{
+	Pipelines.Emplace(
+		LOCTEXT("DefaultPipelineDisplayName","Ingest"),
+		LOCTEXT("DefaultPipelineToolTip", "Download, Convert, Create Assets")
+	);
+
+	SelectedPipeline = *Pipelines.CreateIterator();
+}
+
+const TArray<FPipelineDetails>& FIngestPipelineManager::GetPipelines() const
+{
+	return Pipelines;
+}
+
+TOptional<FPipelineDetails> FIngestPipelineManager::SelectPipelineByDisplayName(const FText& DisplayName)
+{
+	const FPipelineDetails* Pipeline = Pipelines.FindByPredicate(
+		[&DisplayName](const FPipelineDetails& InPipelineDetails)
+		{
+			return InPipelineDetails.DisplayName.CompareTo(DisplayName) == 0;
+		}
+	);
+
+	if (Pipeline)
+	{
+		SelectedPipeline = *Pipeline;
+		return SelectedPipeline;
+	}
+
+	return TOptional<FPipelineDetails>();
+}
+
+const FPipelineDetails& FIngestPipelineManager::GetSelectedPipeline() const
+{
+	return SelectedPipeline;
+}
+
+} // namespace UE::CaptureManager
+
+#undef LOCTEXT_NAMESPACE
\ No newline at end of file
diff --git a/Engine/Plugins/VirtualProduction/CaptureManager/CaptureManagerApp/Source/CaptureManagerEditor/Private/IngestManagement/IngestPipelineManager.h b/Engine/Plugins/VirtualProduction/CaptureManager/CaptureManagerApp/Source/CaptureManagerEditor/Private/IngestManagement/IngestPipelineManager.h
new file mode 100644
--- /dev/null
+++ b/Engine/Plugins/VirtualProduction/CaptureManager/CaptureManagerApp/Source/CaptureManagerEditor/Private/IngestManagement/IngestPipelineManager.h
@@ -0,0 +1,31 @@
+// Copyright Epic Games, Inc. All Rights Reserved.
+
+#pragma once
+
+#include "Containers/Array.h"
+#include "Internationalization/Text.h"
+
+namespace UE::CaptureManager
+{
+
+struct FPipelineDetails
+{
+	FText DisplayName;
+	FText ToolTip;
+};
+
+class FIngestPipelineManager
+{
+public:
+	FIngestPipelineManager();
+
+	const TArray<FPipelineDetails>& GetPipelines() const;
+	TOptional<FPipelineDetails> SelectPipelineByDisplayName(const FText& DisplayName);
+	const FPipelineDetails& GetSelectedPipeline() const;
+
+private:
+	FPipelineDetails SelectedPipeline;
+	TArray<FPipelineDetails> Pipelines;
+};
+
+}
\ No newline at end of file
diff --git a/Engine/Plugins/VirtualProduction/CaptureManager/CaptureManagerApp/Source/CaptureManagerEditor/Private/UI/CaptureManagerPanelController.cpp b/Engine/Plugins/VirtualProduction/CaptureManager/CaptureManagerApp/Source/CaptureManagerEditor/Private/UI/CaptureManagerPanelController.cpp
--- a/Engine/Plugins/VirtualProduction/CaptureManager/CaptureManagerApp/Source/CaptureManagerEditor/Private/UI/CaptureManagerPanelController.cpp
+++ b/Engine/Plugins/VirtualProduction/CaptureManager/CaptureManagerApp/Source/CaptureManagerEditor/Private/UI/CaptureManagerPanelController.cpp
@@ -31,6 +31,7 @@ TSharedPtr<FCaptureManagerPanelController> FCaptureManagerPanelController::MakeI
 }
 
 FCaptureManagerPanelController::FCaptureManagerPanelController(FPrivateToken InPrivateToken) :
+	IngestPipelineManager(MakeShared<UE::CaptureManager::FIngestPipelineManager>()),
 	IngestJobSettingsManager(MakeShared<UE::CaptureManager::FIngestJobSettingsManager>())
 {
 	using namespace UE::CaptureManager;
@@ -81,6 +82,11 @@ TSharedRef<SWidget> FCaptureManagerPanelController::GetIngestJobDetailsWidget()
 	return IngestJobDetailsView.ToSharedRef();
 }
 
+TSharedRef<UE::CaptureManager::FIngestPipelineManager> FCaptureManagerPanelController::GetIngestPipelineManager() const
+{
+	return IngestPipelineManager;
+}
+
 TSharedRef<UE::CaptureManager::FIngestJobSettingsManager> FCaptureManagerPanelController::GetIngestJobSettingsManager() const
 {
 	return IngestJobSettingsManager;
diff --git a/Engine/Plugins/VirtualProduction/CaptureManager/CaptureManagerApp/Source/CaptureManagerEditor/Private/UI/CaptureManagerPanelController.h b/Engine/Plugins/VirtualProduction/CaptureManager/CaptureManagerApp/Source/CaptureManagerEditor/Private/UI/CaptureManagerPanelController.h
--- a/Engine/Plugins/VirtualProduction/CaptureManager/CaptureManagerApp/Source/CaptureManagerEditor/Private/UI/CaptureManagerPanelController.h
+++ b/Engine/Plugins/VirtualProduction/CaptureManager/CaptureManagerApp/Source/CaptureManagerEditor/Private/UI/CaptureManagerPanelController.h
@@ -14,6 +14,7 @@
 namespace UE::CaptureManager
 {
 class FCaptureSourceManager;
+class FIngestPipelineManager;
 class FIngestJobSettingsManager;
 class SIngestJobProcessor;
 class FCaptureEvent;
@@ -46,6 +47,7 @@ public:
 
 	TSharedRef<UE::CaptureManager::SIngestJobProcessor> GetIngestJobProcessorWidget() const;
 	TSharedRef<SWidget> GetIngestJobDetailsWidget() const;
+	TSharedRef<UE::CaptureManager::FIngestPipelineManager> GetIngestPipelineManager() const;
 	TSharedRef<UE::CaptureManager::FIngestJobSettingsManager> GetIngestJobSettingsManager() const;
 
 	TObjectPtr<ULiveLinkDevice> GetCaptureDevice(FGuid InDeviceId) const;
@@ -72,6 +74,7 @@ private:
 	TMap<FGuid, bool> SourcesReachableMap;
 
 	static const int32 NumIngestJobExecutors = 2;
+	TSharedRef<UE::CaptureManager::FIngestPipelineManager> IngestPipelineManager;
 	TSharedRef<UE::CaptureManager::FIngestJobSettingsManager> IngestJobSettingsManager;
 	TSharedPtr<UE::CaptureManager::SIngestJobProcessor> IngestJobProcessorWidget;
 	TSharedPtr<IDetailsView> IngestJobDetailsView;
diff --git a/Engine/Plugins/VirtualProduction/CaptureManager/CaptureManagerApp/Source/CaptureManagerEditor/Private/UI/LiveLinkHubCaptureManagerMode.cpp b/Engine/Plugins/VirtualProduction/CaptureManager/CaptureManagerApp/Source/CaptureManagerEditor/Private/UI/LiveLinkHubCaptureManagerMode.cpp
--- a/Engine/Plugins/VirtualProduction/CaptureManager/CaptureManagerApp/Source/CaptureManagerEditor/Private/UI/LiveLinkHubCaptureManagerMode.cpp
+++ b/Engine/Plugins/VirtualProduction/CaptureManager/CaptureManagerApp/Source/CaptureManagerEditor/Private/UI/LiveLinkHubCaptureManagerMode.cpp
@@ -169,6 +169,7 @@ FLiveLinkHubCaptureManagerMode::FLiveLinkHubCaptureManagerMode(TSharedPtr<FLiveL
 					(
 						FTabManager::NewStack()
 						->SetSizeCoefficient(0.5f)
+						->SetHideTabWell(true)
 						->AddTab("LiveLinkDevices", ETabState::OpenedTab)
 					)
 					->Split
diff --git a/Engine/Plugins/VirtualProduction/CaptureManager/CaptureManagerApp/Source/CaptureManagerEditor/Private/UI/TakesView.cpp b/Engine/Plugins/VirtualProduction/CaptureManager/CaptureManagerApp/Source/CaptureManagerEditor/Private/UI/TakesView.cpp
--- a/Engine/Plugins/VirtualProduction/CaptureManager/CaptureManagerApp/Source/CaptureManagerEditor/Private/UI/TakesView.cpp
+++ b/Engine/Plugins/VirtualProduction/CaptureManager/CaptureManagerApp/Source/CaptureManagerEditor/Private/UI/TakesView.cpp
@@ -14,6 +14,7 @@
 
 #include "Widgets/Input/SButton.h"
 #include "Widgets/Input/SComboBox.h"
+#include "SPositiveActionButton.h"
 
 #include "Ingest/LiveLinkDeviceCapability_Ingest.h"
 #include "Ingest/IngestCapability_Events.h"
@@ -23,7 +24,7 @@
 
 DEFINE_LOG_CATEGORY_STATIC(LogTakesView, Log, All);
 
-static const FText JobPresetText = LOCTEXT("JobPresetText", "Job Preset");
+static const FText JobPipelineText = LOCTEXT("JobPipelineText", "Pipeline");
 
 static const FString OuterPackageName = TEXT("CaptureDevices");
 
@@ -31,6 +32,7 @@ STakesView::STakesView() = default;
 
 void STakesView::Construct(
 	const FArguments& InArgs,
+	TSharedRef<UE::CaptureManager::FIngestPipelineManager> InIngestPipelineManager,
 	TSharedRef<UE::CaptureManager::FIngestJobSettingsManager> InIngestJobsSettingsManager,
 	TUniquePtr<FGetCurrentSelectionDelegate> InGetCurrentSelectionDelegate
 )
@@ -38,21 +40,22 @@ void STakesView::Construct(
 	GetCurrentSelectionDelegate = MoveTemp(InGetCurrentSelectionDelegate);
 	check(GetCurrentSelectionDelegate);
 
+	IngestPipelineManager = InIngestPipelineManager.ToSharedPtr();
 	IngestJobSettingsManager = InIngestJobsSettingsManager.ToSharedPtr();
 	OnAddTakesToIngestQueue = InArgs._OnAddTakesToIngestQueue;
+	CurrentPipeline = IngestPipelineManager->GetSelectedPipeline();
 	CurrentPresetName = IngestJobSettingsManager->GetSelectedPreset().DisplayName;
 
-	TArray<FText> PresetDisplayNames = IngestJobSettingsManager->GetPresetDisplayNames();
-	PresetDisplayNames.Sort(FText::FSortPredicate());
-
-	for (const FText& PresetDisplayName : PresetDisplayNames)
+	TArray<UE::CaptureManager::FPipelineDetails> Pipelines = IngestPipelineManager->GetPipelines();
+	for (const UE::CaptureManager::FPipelineDetails& Pipeline : Pipelines)
 	{
-		IngestJobPresetNames.Emplace(MakeShared<FText>(PresetDisplayName));
+		PipelineNames.Emplace(MakeShared<FText>(Pipeline.DisplayName));
 	}
 
-	AddToQueueButton = SNew(SButton)
-		.HAlign(EHorizontalAlignment::HAlign_Center)
+	AddToQueueButton = SNew(SPositiveActionButton)
+		.Icon(FAppStyle::Get().GetBrush("Icons.Plus"))
 		.Text(LOCTEXT("AddToQueueText", "Add to Queue"))
+		.ToolTipText(LOCTEXT("AddToQueue_Tooltip", "Add selected takes to the ingest jobs list"))
 		.OnClicked_Raw(this, &STakesView::OnAddToQueueButtonClicked);
 
 	IContentBrowserSingleton& ContentBrowser = FModuleManager::LoadModuleChecked<FContentBrowserModule>("ContentBrowser").Get();
@@ -62,10 +65,6 @@ void STakesView::Construct(
 	ChildSlot
 		[
 			SNew(SVerticalBox)
-				+ SVerticalBox::Slot()
-				[
-					ContentBrowser.CreateAssetPicker(TakesPickerConfig)
-				]
 				+ SVerticalBox::Slot()
 				.VAlign(EVerticalAlignment::VAlign_Bottom)
 				.AutoHeight()
@@ -78,7 +77,7 @@ void STakesView::Construct(
 						.VAlign(EVerticalAlignment::VAlign_Center)
 						[
 							SNew(STextBlock)
-								.Text(JobPresetText)
+								.Text(JobPipelineText)
 						]
 						+ SHorizontalBox::Slot()
 						.Padding(2.0f)
@@ -86,9 +85,9 @@ void STakesView::Construct(
 						.VAlign(EVerticalAlignment::VAlign_Center)
 						[
 							SNew(SComboBox<TSharedRef<FText>>)
-								.OptionsSource(&IngestJobPresetNames)
-								.OnGenerateWidget_Raw(this, &STakesView::OnGeneratePresetNameWidget)
-								.OnSelectionChanged_Raw(this, &STakesView::OnPresetSelectionChanged)
+								.OptionsSource(&PipelineNames)
+								.OnGenerateWidget_Raw(this, &STakesView::OnGeneratePipelineNameWidget)
+								.OnSelectionChanged_Raw(this, &STakesView::OnPipelineSelectionChanged)
 								.Content()
 								[
 									SNew(STextBlock)
@@ -96,7 +95,13 @@ void STakesView::Construct(
 										.Text_Lambda(
 											[this]()
 											{
-												return CurrentPresetName;
+												return CurrentPipeline.DisplayName;
+											}
+										)
+										.ToolTipText_Lambda(
+											[this]()
+											{
+												return CurrentPipeline.ToolTip;
 											}
 										)
 								]
@@ -109,6 +114,10 @@ void STakesView::Construct(
 							AddToQueueButton.ToSharedRef()
 						]
 				]
+				+ SVerticalBox::Slot()
+				[
+					ContentBrowser.CreateAssetPicker(TakesPickerConfig)
+				]
 		];
 }
 
@@ -174,21 +183,18 @@ FReply STakesView::OnAddToQueueButtonClicked()
 	return FReply::Handled();
 }
 
-TSharedRef<SWidget> STakesView::OnGeneratePresetNameWidget(TSharedRef<FText> InText)
+TSharedRef<SWidget> STakesView::OnGeneratePipelineNameWidget(TSharedRef<FText> InText)
 {
 	return SNew(STextBlock).Text(*InText);
 }
 
-void STakesView::OnPresetSelectionChanged(TSharedPtr<FText> InText, ESelectInfo::Type InSelectType)
+void STakesView::OnPipelineSelectionChanged(TSharedPtr<FText> InText, ESelectInfo::Type InSelectType)
 {
-	check(IngestJobSettingsManager);
+	check(IngestPipelineManager);
 
-	if (InText)
+	if (TOptional<UE::CaptureManager::FPipelineDetails> SelectedPipeline = IngestPipelineManager->SelectPipelineByDisplayName(*InText))
 	{
-		if (IngestJobSettingsManager->SelectPresetByDisplayName(*InText))
-		{
-			CurrentPresetName = *InText;
-		}
+		CurrentPipeline = SelectedPipeline.GetValue();
 	}
 }
 
@@ -302,6 +308,7 @@ void FTakesView::CreateTakesView()
 		return;
 	}
 
+	TSharedRef<UE::CaptureManager::FIngestPipelineManager> IngestPipelineManager = Controller->GetIngestPipelineManager();
 	TSharedRef<UE::CaptureManager::FIngestJobSettingsManager> IngestJobSettingsManager = Controller->GetIngestJobSettingsManager();
 
 	ILiveLinkDeviceModule& Module = ILiveLinkDeviceModule::Get();
@@ -347,7 +354,7 @@ void FTakesView::CreateTakesView()
 				});
 	}
 
-	SAssignNew(TakesTileView, STakesView, MoveTemp(IngestJobSettingsManager), MoveTemp(GetCurrentSelectionDelegate))
+	SAssignNew(TakesTileView, STakesView, MoveTemp(IngestPipelineManager), MoveTemp(IngestJobSettingsManager), MoveTemp(GetCurrentSelectionDelegate))
 		.OnAddTakesToIngestQueue_Raw(this, &FTakesView::AddTakesToIngestQueue)
 		.TakesPickerConfig(TakesPickerConfig);
 }
diff --git a/Engine/Plugins/VirtualProduction/CaptureManager/CaptureManagerApp/Source/CaptureManagerEditor/Private/UI/TakesView.h b/Engine/Plugins/VirtualProduction/CaptureManager/CaptureManagerApp/Source/CaptureManagerEditor/Private/UI/TakesView.h
--- a/Engine/Plugins/VirtualProduction/CaptureManager/CaptureManagerApp/Source/CaptureManagerEditor/Private/UI/TakesView.h
+++ b/Engine/Plugins/VirtualProduction/CaptureManager/CaptureManagerApp/Source/CaptureManagerEditor/Private/UI/TakesView.h
@@ -15,10 +15,12 @@
 #include "IContentBrowserSingleton.h"
 
 #include "IngestManagement/IngestJobSettingsManager.h"
+#include "IngestManagement/IngestPipelineManager.h"
 
 class FCaptureManagerPanelController;
 class FCaptureSourcesView;
 class SEditableTextBox;
+class SPositiveActionButton;
 class SSearchBox;
 class SSlateBoxBrush;
 class ULiveLinkDevice;
@@ -42,7 +44,8 @@ public:
 	~STakesView() = default;
 
 	void Construct(
-		const FArguments& InArgs, 
+		const FArguments& InArgs,
+		TSharedRef<UE::CaptureManager::FIngestPipelineManager> InIngestPipelineManager,
 		TSharedRef<UE::CaptureManager::FIngestJobSettingsManager> IngestJobsSettingsManager,
 		TUniquePtr<FGetCurrentSelectionDelegate> GetCurrentSelectionDelegate
 	);
@@ -50,15 +53,17 @@ public:
 
 private:
 	FReply OnAddToQueueButtonClicked();
-	TSharedRef<SWidget> OnGeneratePresetNameWidget(TSharedRef<FText> InText);
-	void OnPresetSelectionChanged(TSharedPtr<FText> InText, ESelectInfo::Type InSelectType);
+	TSharedRef<SWidget> OnGeneratePipelineNameWidget(TSharedRef<FText> InText);
+	void OnPipelineSelectionChanged(TSharedPtr<FText> InText, ESelectInfo::Type InSelectType);
 	TArray<TStrongObjectPtr<UTakeVirtualAsset>> GetSelectedTakeAssets() const;
 
 	TUniquePtr<FGetCurrentSelectionDelegate> GetCurrentSelectionDelegate;
-	TSharedPtr<SButton> AddToQueueButton;
+	TSharedPtr<SPositiveActionButton> AddToQueueButton;
 	FOnAddTakesToIngestQueue OnAddTakesToIngestQueue;
 
-	TArray<TSharedRef<FText>> IngestJobPresetNames;
+	TSharedPtr<UE::CaptureManager::FIngestPipelineManager> IngestPipelineManager;
+	TArray<TSharedRef<FText>> PipelineNames;
+	UE::CaptureManager::FPipelineDetails CurrentPipeline;
 	TSharedPtr<UE::CaptureManager::FIngestJobSettingsManager> IngestJobSettingsManager;
 	FText CurrentPresetName;
 };
