CaptureManager: Allow the configuration of the number of parallel ingest jobs.

#rnx

diff --git a/Engine/Plugins/VirtualProduction/CaptureManager/CaptureManagerApp/Source/CaptureManagerEditor/Private/UI/CaptureManagerPanelController.cpp b/Engine/Plugins/VirtualProduction/CaptureManager/CaptureManagerApp/Source/CaptureManagerEditor/Private/UI/CaptureManagerPanelController.cpp
--- a/Engine/Plugins/VirtualProduction/CaptureManager/CaptureManagerApp/Source/CaptureManagerEditor/Private/UI/CaptureManagerPanelController.cpp
+++ b/Engine/Plugins/VirtualProduction/CaptureManager/CaptureManagerApp/Source/CaptureManagerEditor/Private/UI/CaptureManagerPanelController.cpp
@@ -10,10 +10,8 @@
 
 #include "IDetailsView.h"
 #include "PropertyEditorModule.h"
-#include "Serialization/BufferArchive.h"
 
-#include "Misc/FileHelper.h"
-#include "Async/Event.h"
+#include "Settings/CaptureManagerSettings.h"
 
 #include "Engine/Engine.h"
 #include "Ingest/LiveLinkDeviceCapability_Ingest.h"
@@ -23,6 +21,25 @@
 
 DEFINE_LOG_CATEGORY_STATIC(LogCaptureManagerPanelController, Log, All);
 
+namespace UE::CaptureManager::Private
+{
+static int32 GetNumberOfIngestExecutors()
+{
+	const UCaptureManagerSettings* Settings = GetDefault<UCaptureManagerSettings>();
+	check(Settings);
+
+	// Fallback for safety, just in case we can't find the settings
+	int32 NumExecutors = 1;
+
+	if (Settings)
+	{
+		NumExecutors = Settings->NumIngestExecutors;
+	}
+
+	return NumExecutors;
+}
+}
+
 TSharedPtr<FCaptureManagerPanelController> FCaptureManagerPanelController::MakeInstance()
 {
 	// We do this to enforce the use of a TSharedPtr and to hide CreateViews from the user
@@ -67,8 +84,10 @@ void FCaptureManagerPanelController::CreateViews()
 void FCaptureManagerPanelController::CreateIngestManagementViews()
 {
 	using namespace UE::CaptureManager;
+	using namespace UE::CaptureManager::Private;
 
-	IngestJobProcessorWidget = SNew(SIngestJobProcessor).NumExecutors(NumIngestJobExecutors);
+	const int32 NumExecutors = GetNumberOfIngestExecutors();
+	IngestJobProcessorWidget = SNew(SIngestJobProcessor).NumExecutors(NumExecutors);
 	IngestJobDetailsView = CreateIngestJobDetailsView().ToSharedPtr();
 	FinishedChangingPropertiesHandle = IngestJobDetailsView->OnFinishedChangingProperties().AddRaw(this, &FCaptureManagerPanelController::OnFinishedEditingJobProperties);
 
diff --git a/Engine/Plugins/VirtualProduction/CaptureManager/CaptureManagerApp/Source/CaptureManagerEditor/Private/UI/CaptureManagerPanelController.h b/Engine/Plugins/VirtualProduction/CaptureManager/CaptureManagerApp/Source/CaptureManagerEditor/Private/UI/CaptureManagerPanelController.h
--- a/Engine/Plugins/VirtualProduction/CaptureManager/CaptureManagerApp/Source/CaptureManagerEditor/Private/UI/CaptureManagerPanelController.h
+++ b/Engine/Plugins/VirtualProduction/CaptureManager/CaptureManagerApp/Source/CaptureManagerEditor/Private/UI/CaptureManagerPanelController.h
@@ -75,7 +75,6 @@ private:
 
 	TMap<FGuid, bool> SourcesReachableMap;
 
-	static const int32 NumIngestJobExecutors = 2;
 	TSharedRef<UE::CaptureManager::FIngestPipelineManager> IngestPipelineManager;
 	TSharedRef<UE::CaptureManager::FUnrealEndpointManager> UnrealEndpointManager;
 	TSharedRef<UE::CaptureManager::FIngestJobSettingsManager> IngestJobSettingsManager;
diff --git a/Engine/Plugins/VirtualProduction/CaptureManager/CaptureManagerApp/Source/CaptureManagerSettings/Public/Settings/CaptureManagerSettings.h b/Engine/Plugins/VirtualProduction/CaptureManager/CaptureManagerApp/Source/CaptureManagerSettings/Public/Settings/CaptureManagerSettings.h
--- a/Engine/Plugins/VirtualProduction/CaptureManager/CaptureManagerApp/Source/CaptureManagerSettings/Public/Settings/CaptureManagerSettings.h
+++ b/Engine/Plugins/VirtualProduction/CaptureManager/CaptureManagerApp/Source/CaptureManagerSettings/Public/Settings/CaptureManagerSettings.h
@@ -80,6 +80,10 @@ public:
 	UPROPERTY(config, EditAnywhere, Category = "Upload")
 	FString DefaultUploadHostName;
 
+	/** The number of jobs to run in parallel. Requires a restart of Live Link Hub (Warning: Setting this too high could degrade performance) */
+	UPROPERTY(config, EditAnywhere, Category = "General", meta = (DisplayName = "Parallel Jobs", UIMin = 1, ClampMin = 1, UIMax = 64, ClampMax = 64))
+	int32 NumIngestExecutors = 2;
+
 private:
 
 	virtual void PostEditChangeProperty(struct FPropertyChangedEvent& InPropertyChangedEvent) override;
