Make take names unique

diff --git a/Engine/Plugins/VirtualProduction/CaptureManager/CaptureManagerApp/Source/CaptureManagerEditor/Private/UI/TakesView.cpp b/Engine/Plugins/VirtualProduction/CaptureManager/CaptureManagerApp/Source/CaptureManagerEditor/Private/UI/TakesView.cpp
--- a/Engine/Plugins/VirtualProduction/CaptureManager/CaptureManagerApp/Source/CaptureManagerEditor/Private/UI/TakesView.cpp
+++ b/Engine/Plugins/VirtualProduction/CaptureManager/CaptureManagerApp/Source/CaptureManagerEditor/Private/UI/TakesView.cpp
@@ -568,10 +568,9 @@ TArray<TStrongObjectPtr<UTakeVirtualAsset>> FTakesView::MakeTakeObjects(FGuid In
 
 		FTakeMetadata TakeInfo = TakeInfoOpt.GetValue();
 
-		FString DesiredNameStr = FString::Format(TEXT("{0}_{1}"), { TakeInfo.Slate, TakeInfo.TakeNumber });
-		FName DesiredName(MoveTemp(DesiredNameStr));
+		FName Name = MakeUniqueTakeName(TakeInfo, Package);
+		TStrongObjectPtr<UTakeVirtualAsset> TakeObject(NewObject<UTakeVirtualAsset>(Package, Name));
 
-		TStrongObjectPtr<UTakeVirtualAsset> TakeObject(NewObject<UTakeVirtualAsset>(Package, DesiredName));
 		TakeObject->CaptureDeviceId = InCaptureDeviceId;
 		TakeObject->TakeId = TakeId;
 		TakeObject->Metadata = MoveTemp(TakeInfo);
@@ -613,5 +612,19 @@ TOptional<FTakeThumbnail> FTakesView::CreateThumbnail(const FTakeThumbnailData&
 	return FTakeThumbnail(PreviewImageTexture);
 }
 
+FName FTakesView::MakeUniqueTakeName(FTakeMetadata TakeInfo, UPackage* Package)
+{
+	FName UniqueName(FString::Format(TEXT("{0}_{1}"), { TakeInfo.Slate, TakeInfo.TakeNumber }));
+	UObject* FoundObject = StaticFindObjectFastInternal(nullptr, Package, UniqueName);
+
+	int32 UniquePartIncrement = 0;
+	while (FoundObject)
+	{
+		FName MaybeUniqueName(FString::Format(TEXT("{0}_{1}_{2}"), { TakeInfo.Slate, TakeInfo.TakeNumber, ++UniquePartIncrement }));
+		FoundObject = StaticFindObjectFastInternal(nullptr, Package, MaybeUniqueName);
+		UniqueName = MaybeUniqueName;
+	}
+	return UniqueName;
+}
 
 #undef LOCTEXT_NAMESPACE
diff --git a/Engine/Plugins/VirtualProduction/CaptureManager/CaptureManagerApp/Source/CaptureManagerEditor/Private/UI/TakesView.h b/Engine/Plugins/VirtualProduction/CaptureManager/CaptureManagerApp/Source/CaptureManagerEditor/Private/UI/TakesView.h
--- a/Engine/Plugins/VirtualProduction/CaptureManager/CaptureManagerApp/Source/CaptureManagerEditor/Private/UI/TakesView.h
+++ b/Engine/Plugins/VirtualProduction/CaptureManager/CaptureManagerApp/Source/CaptureManagerEditor/Private/UI/TakesView.h
@@ -95,4 +95,6 @@ private:
 	TArray<TStrongObjectPtr<UTakeVirtualAsset>> MakeTakeObjects(FGuid InCaptureDeviceId,
 																TOptional<TArray<UE::CaptureManager::FTakeId>> InTakeIds = TOptional<TArray<UE::CaptureManager::FTakeId>>());
 	TOptional<FTakeThumbnail> CreateThumbnail(const FTakeThumbnailData& InThumbnailData);
+
+	FName MakeUniqueTakeName(FTakeMetadata TakeInfo, UPackage* Package);
 };
