[CaptureManager] Set default iphone rotation correctly

diff --git a/Engine/Plugins/VirtualProduction/CaptureManager/CaptureManagerApp/Source/CaptureManagerEditor/Private/IngestManagement/IngestJob.h b/Engine/Plugins/VirtualProduction/CaptureManager/CaptureManagerApp/Source/CaptureManagerEditor/Private/IngestManagement/IngestJob.h
--- a/Engine/Plugins/VirtualProduction/CaptureManager/CaptureManagerApp/Source/CaptureManagerEditor/Private/IngestManagement/IngestJob.h
+++ b/Engine/Plugins/VirtualProduction/CaptureManager/CaptureManagerApp/Source/CaptureManagerEditor/Private/IngestManagement/IngestJob.h
@@ -24,7 +24,7 @@ public:
 	{
 		struct FVideoSettings
 		{
-			FString Format;
+			::EOutputImageFormat Format;
 			FString FileNamePrefix;
 			::EImagePixelFormat ImagePixelFormat;
 			::EImageRotation ImageRotation;
@@ -32,7 +32,7 @@ public:
 
 		struct FAudioSettings
 		{
-			FString Format;
+			::EAudioFormat Format;
 			FString FileNamePrefix;
 		};
 
diff --git a/Engine/Plugins/VirtualProduction/CaptureManager/CaptureManagerApp/Source/CaptureManagerEditor/Private/IngestManagement/IngestJobExecutor.cpp b/Engine/Plugins/VirtualProduction/CaptureManager/CaptureManagerApp/Source/CaptureManagerEditor/Private/IngestManagement/IngestJobExecutor.cpp
--- a/Engine/Plugins/VirtualProduction/CaptureManager/CaptureManagerApp/Source/CaptureManagerEditor/Private/IngestManagement/IngestJobExecutor.cpp
+++ b/Engine/Plugins/VirtualProduction/CaptureManager/CaptureManagerApp/Source/CaptureManagerEditor/Private/IngestManagement/IngestJobExecutor.cpp
@@ -231,7 +231,7 @@ UIngestCapability_ProcessHandle* FIngestJobExecutor::RunIngest(FIngestJob& InJob
 			FNamingTokenResultData VideoPrefixResult = NamingTokensSubsystem->EvaluateTokenString(FileNamePrefix, GeneralNamingTokenArgs);
 			FIngestCapability_VideoOptions VideoOptions;
 			VideoOptions.FileNamePrefix = VideoPrefixResult.EvaluatedText.ToString();
-			VideoOptions.Format = Settings.VideoSettings.Format;
+			VideoOptions.Format = StaticEnum<EOutputImageFormat>()->GetDisplayNameTextByValue(static_cast<std::underlying_type_t<EOutputImageFormat>>(Settings.VideoSettings.Format)).ToString();
 			VideoOptions.PixelFormat = ConvertImagePixelFormat(Settings.VideoSettings.ImagePixelFormat);
 			VideoOptions.Rotation = ConvertImageRotation(Settings.VideoSettings.ImageRotation);
 			IngestOptions->Video = VideoOptions;
@@ -240,7 +240,7 @@ UIngestCapability_ProcessHandle* FIngestJobExecutor::RunIngest(FIngestJob& InJob
 			FNamingTokenResultData AudioPrefixResult = NamingTokensSubsystem->EvaluateTokenString(AudioPrefix, GeneralNamingTokenArgs);
 			FIngestCapability_AudioOptions AudioOptions;
 			AudioOptions.FileNamePrefix = AudioPrefixResult.EvaluatedText.ToString();
-			AudioOptions.Format = Settings.AudioSettings.Format;
+			AudioOptions.Format = StaticEnum<EAudioFormat>()->GetDisplayNameTextByValue(static_cast<std::underlying_type_t<EAudioFormat>>(Settings.AudioSettings.Format)).ToString();
 			IngestOptions->Audio = AudioOptions;
 
 			IngestOptions->UploadHostName = Settings.UploadHostName;
diff --git a/Engine/Plugins/VirtualProduction/CaptureManager/CaptureManagerApp/Source/CaptureManagerEditor/Private/IngestManagement/IngestJobSettingsManager.cpp b/Engine/Plugins/VirtualProduction/CaptureManager/CaptureManagerApp/Source/CaptureManagerEditor/Private/IngestManagement/IngestJobSettingsManager.cpp
--- a/Engine/Plugins/VirtualProduction/CaptureManager/CaptureManagerApp/Source/CaptureManagerEditor/Private/IngestManagement/IngestJobSettingsManager.cpp
+++ b/Engine/Plugins/VirtualProduction/CaptureManager/CaptureManagerApp/Source/CaptureManagerEditor/Private/IngestManagement/IngestJobSettingsManager.cpp
@@ -67,12 +67,21 @@ TStrongObjectPtr<const UIngestJobSettingsBase> FIngestJobSettingsManager::GetSel
 	return TStrongObjectPtr<const UIngestJobSettingsBase>(SelectedPreset);
 }
 
-void FIngestJobSettingsManager::CreateSettingsFromSelectedPreset(const FGuid& JobGuid)
+void FIngestJobSettingsManager::ApplyJobSpecificSettings(const FGuid& InJobGuid, const FIngestJob::FSettings& InSettings)
 {
 	FScopeLock Lock(&CriticalSection);
 
 	TStrongObjectPtr<UIngestJobSettingsBase>& Object = Settings.Emplace_GetRef(DuplicateObject(SelectedPreset, nullptr));
-	Object->JobGuid = JobGuid;
+	Object->JobGuid = InJobGuid;
+	Object->AudioFileNamePrefix = InSettings.AudioSettings.FileNamePrefix;
+	Object->AudioFormat = InSettings.AudioSettings.Format;
+	Object->UploadHostName = InSettings.UploadHostName;
+	Object->WorkingDirectory = FDirectoryPath(InSettings.WorkingDirectory);
+	Object->DownloadFolder = FDirectoryPath(InSettings.DownloadFolder);
+	Object->ImageFileNamePrefix = InSettings.VideoSettings.FileNamePrefix;
+	Object->ImageFormat = InSettings.VideoSettings.Format;
+	Object->ImagePixelFormat = InSettings.VideoSettings.ImagePixelFormat;
+	Object->ImageRotation = InSettings.VideoSettings.ImageRotation;
 }
 
 int32 FIngestJobSettingsManager::RemoveSettings(const TArray<FGuid>& JobGuids)
diff --git a/Engine/Plugins/VirtualProduction/CaptureManager/CaptureManagerApp/Source/CaptureManagerEditor/Private/IngestManagement/IngestJobSettingsManager.h b/Engine/Plugins/VirtualProduction/CaptureManager/CaptureManagerApp/Source/CaptureManagerEditor/Private/IngestManagement/IngestJobSettingsManager.h
--- a/Engine/Plugins/VirtualProduction/CaptureManager/CaptureManagerApp/Source/CaptureManagerEditor/Private/IngestManagement/IngestJobSettingsManager.h
+++ b/Engine/Plugins/VirtualProduction/CaptureManager/CaptureManagerApp/Source/CaptureManagerEditor/Private/IngestManagement/IngestJobSettingsManager.h
@@ -4,6 +4,8 @@
 
 #include "UIngestJobSettings.h"
 
+#include "IngestJob.h"
+
 namespace UE::CaptureManager
 {
 
@@ -16,7 +18,7 @@ public:
 	bool SelectPresetByDisplayName(const FText& DisplayName);
 
 	TStrongObjectPtr<const UIngestJobSettingsBase> GetSelectedPreset() const;
-	void CreateSettingsFromSelectedPreset(const FGuid& InJobGuid);
+	void ApplyJobSpecificSettings(const FGuid& InJobGuid, const FIngestJob::FSettings& InSettings);
 	TArray<TWeakObjectPtr<UIngestJobSettingsBase>> GetSettings(const TArray<FGuid>& InJobGuids) const;
 
 	int32 RemoveSettings(const TArray<FGuid>& InJobGuids);
diff --git a/Engine/Plugins/VirtualProduction/CaptureManager/CaptureManagerApp/Source/CaptureManagerEditor/Private/UI/CaptureManagerPanelController.cpp b/Engine/Plugins/VirtualProduction/CaptureManager/CaptureManagerApp/Source/CaptureManagerEditor/Private/UI/CaptureManagerPanelController.cpp
--- a/Engine/Plugins/VirtualProduction/CaptureManager/CaptureManagerApp/Source/CaptureManagerEditor/Private/UI/CaptureManagerPanelController.cpp
+++ b/Engine/Plugins/VirtualProduction/CaptureManager/CaptureManagerApp/Source/CaptureManagerEditor/Private/UI/CaptureManagerPanelController.cpp
@@ -218,7 +218,7 @@ void FCaptureManagerPanelController::OnJobsAdded(TArray<TSharedRef<UE::CaptureMa
 
 	for (const TSharedRef<FIngestJob>& IngestJob : IngestJobs)
 	{
-		IngestJobSettingsManager->CreateSettingsFromSelectedPreset(IngestJob->GetGuid());
+		IngestJobSettingsManager->ApplyJobSpecificSettings(IngestJob->GetGuid(), IngestJob->GetSettings());
 	}
 }
 
@@ -286,14 +286,14 @@ void FCaptureManagerPanelController::OnFinishedEditingJobProperties(const FPrope
 				.DownloadFolder = Settings->DownloadFolder.Path,
 				.VideoSettings =
 				{
-					.Format = StaticEnum<EOutputImageFormat>()->GetDisplayNameTextByValue(static_cast<std::underlying_type_t<EOutputImageFormat>>(Settings->ImageFormat)).ToString(),
+					.Format = Settings->ImageFormat,
 					.FileNamePrefix = Settings->ImageFileNamePrefix,
 					.ImagePixelFormat = Settings->ImagePixelFormat,
 					.ImageRotation = Settings->ImageRotation
 				},
 				.AudioSettings =
 				{
-					.Format = StaticEnum<EAudioFormat>()->GetDisplayNameTextByValue(static_cast<std::underlying_type_t<EAudioFormat>>(Settings->AudioFormat)).ToString(),
+					.Format = Settings->AudioFormat,
 					.FileNamePrefix = Settings->AudioFileNamePrefix
 				},
 				.UploadHostName = Settings->UploadHostName
diff --git a/Engine/Plugins/VirtualProduction/CaptureManager/CaptureManagerApp/Source/CaptureManagerEditor/Private/UI/TakesView.cpp b/Engine/Plugins/VirtualProduction/CaptureManager/CaptureManagerApp/Source/CaptureManagerEditor/Private/UI/TakesView.cpp
--- a/Engine/Plugins/VirtualProduction/CaptureManager/CaptureManagerApp/Source/CaptureManagerEditor/Private/UI/TakesView.cpp
+++ b/Engine/Plugins/VirtualProduction/CaptureManager/CaptureManagerApp/Source/CaptureManagerEditor/Private/UI/TakesView.cpp
@@ -259,6 +259,22 @@ FTakesView::FTakesView(TWeakPtr<FCaptureManagerPanelController> InController) :
 	CreateTakesView();
 }
 
+EImageRotation FTakesView::ConvertOrientation(TOptional<FTakeMetadata::FVideo::EOrientation> InOrientation) const
+{
+	switch (InOrientation.Get(FTakeMetadata::FVideo::EOrientation::Original))
+	{
+		case FTakeMetadata::FVideo::EOrientation::CW90:
+			return EImageRotation::CW_270;
+		case FTakeMetadata::FVideo::EOrientation::CW180:
+			return EImageRotation::CW_180;
+		case FTakeMetadata::FVideo::EOrientation::CW270:
+			return EImageRotation::CW_90;
+		case FTakeMetadata::FVideo::EOrientation::Original:
+		default:
+			return EImageRotation::None;
+	}
+}
+
 void FTakesView::AddTakesToIngestQueue(const TArray<TStrongObjectPtr<UTakeVirtualAsset>>& InTakeAssets, const UIngestJobSettingsBase& InJobPreset)
 {
 	using namespace UE::CaptureManager;
@@ -276,6 +292,10 @@ void FTakesView::AddTakesToIngestQueue(const TArray<TStrongObjectPtr<UTakeVirtua
 
 	for (const TStrongObjectPtr<UTakeVirtualAsset>& TakeAsset : InTakeAssets)
 	{
+
+		EImageRotation ImageRotationToApply =
+			TakeAsset->Metadata.Video.IsEmpty() ? InJobPreset.ImageRotation : ConvertOrientation(TakeAsset->Metadata.Video[0].Orientation);
+
 		// We convert the UObject based settings into something that the ingest job can truly own
 		FIngestJob::FSettings IngestJobSettings =
 		{
@@ -283,13 +303,14 @@ void FTakesView::AddTakesToIngestQueue(const TArray<TStrongObjectPtr<UTakeVirtua
 			.DownloadFolder = InJobPreset.DownloadFolder.Path,
 			.VideoSettings =
 			{
-				.Format = StaticEnum<EOutputImageFormat>()->GetDisplayNameTextByValue(static_cast<std::underlying_type_t<EOutputImageFormat>>(InJobPreset.ImageFormat)).ToString(),
+				.Format = InJobPreset.ImageFormat,
 				.FileNamePrefix = InJobPreset.ImageFileNamePrefix,
-				.ImagePixelFormat = InJobPreset.ImagePixelFormat
+				.ImagePixelFormat = InJobPreset.ImagePixelFormat,
+				.ImageRotation = ImageRotationToApply
 			},
 			.AudioSettings =
 			{
-				.Format = StaticEnum<EAudioFormat>()->GetDisplayNameTextByValue(static_cast<std::underlying_type_t<EAudioFormat>>(InJobPreset.AudioFormat)).ToString(),
+				.Format = InJobPreset.AudioFormat,
 				.FileNamePrefix = InJobPreset.AudioFileNamePrefix
 			},
 			.UploadHostName = InJobPreset.UploadHostName
diff --git a/Engine/Plugins/VirtualProduction/CaptureManager/CaptureManagerApp/Source/CaptureManagerEditor/Private/UI/TakesView.h b/Engine/Plugins/VirtualProduction/CaptureManager/CaptureManagerApp/Source/CaptureManagerEditor/Private/UI/TakesView.h
--- a/Engine/Plugins/VirtualProduction/CaptureManager/CaptureManagerApp/Source/CaptureManagerEditor/Private/UI/TakesView.h
+++ b/Engine/Plugins/VirtualProduction/CaptureManager/CaptureManagerApp/Source/CaptureManagerEditor/Private/UI/TakesView.h
@@ -90,6 +90,8 @@ private:
 
 	TArray<FAssetData> CreateTakeAssetsDataForCaptureDevice(TOptional<FString> InCaptureDeviceName = TOptional<FString>());
 	TArray<FAssetData> CreateTakeAssetsData(const TMap<FString, TArray<TStrongObjectPtr<UTakeVirtualAsset>>>& InSelectedCaptureDevicesTakesMap);
+
+	EImageRotation ConvertOrientation(TOptional<FTakeMetadata::FVideo::EOrientation> InOrientation) const;
 	void AddTakesToIngestQueue(const TArray<TStrongObjectPtr<UTakeVirtualAsset>>& InTakeAssets, const UIngestJobSettingsBase& InJobPreset);
 
 	TArray<TStrongObjectPtr<UTakeVirtualAsset>> MakeTakeObjects(FGuid InCaptureDeviceId,
diff --git a/Engine/Plugins/VirtualProduction/CaptureManager/CaptureManagerApp/Source/LiveLinkFaceMetadata/Private/LiveLinkFaceMetadata.cpp b/Engine/Plugins/VirtualProduction/CaptureManager/CaptureManagerApp/Source/LiveLinkFaceMetadata/Private/LiveLinkFaceMetadata.cpp
--- a/Engine/Plugins/VirtualProduction/CaptureManager/CaptureManagerApp/Source/LiveLinkFaceMetadata/Private/LiveLinkFaceMetadata.cpp
+++ b/Engine/Plugins/VirtualProduction/CaptureManager/CaptureManagerApp/Source/LiveLinkFaceMetadata/Private/LiveLinkFaceMetadata.cpp
@@ -153,6 +153,7 @@ public:
 
 	static EMediaOrientation ParseOrientation(int32 InOrientation = 4);
 	static bool ParseVideoMetadata(const FString& InTakeDirectory, FLiveLinkFaceVideoMetadata& OutVideoMetadata);
+	static bool ParseVideoMetadataFromString(const FString& InJsonString, FLiveLinkFaceVideoMetadata& OutVideoMetadata);
 	static bool ParseDepthMetadata(const FString& InTakeDirectory, FLiveLinkFaceDepthMetadata& OutDepthMetadata);
 	static bool ParseAudioMetadata(const FString& InTakeDirectory, FLiveLinkFaceAudioMetadata& OutAudioMetadata);
 	static bool ParseTakeInfo(const FString& InTakeDirectory, FLiveLinkFaceTakeInfo& OutTakeInfo);
@@ -167,7 +168,9 @@ private:
 	static FLiveLinkTakeMetadataVoidResult ParseBool(const TSharedPtr<FJsonObject> JsonObject, const FString& Key, bool& OutBool);
 	static FLiveLinkTakeMetadataError CreateErrorForMissingJsonKey(const FString& Key);
 	static TArray<TSharedPtr<class FJsonValue>> ParseJsonArrayFromFile(const FString& InFilePath);
-	static TSharedPtr<class FJsonObject> ParseJsonObjectFromFile(const FString& InFilePath);
+	//static TSharedPtr<class FJsonObject> ParseJsonObjectFromFile(const FString& InFilePath);
+	static const FString ParseJsonStringFromFile(const FString& InFilePath);
+	static TSharedPtr<class FJsonObject> ParseJsonObjectFromString(const FString& InJsonString);
 };
 
 static bool ContainsWhitespace(const FString& InStr)
@@ -316,8 +319,14 @@ EMediaOrientation FLiveLinkFaceMetadataParser::ParseOrientation(int32 InOrientat
 bool FLiveLinkFaceMetadataParser::ParseVideoMetadata(const FString& InTakeDirectory, FLiveLinkFaceVideoMetadata& OutVideoMetadata)
 {
 	const FString& VideoMetadataFile = InTakeDirectory / FLiveLinkFaceStaticFileNames::VideoMetadata;
+	const FString& VideoMetadataString = ParseJsonStringFromFile(VideoMetadataFile);
 
-	if (TSharedPtr<FJsonObject> VideoMetadataJson = ParseJsonObjectFromFile(VideoMetadataFile))
+	return !VideoMetadataString.IsEmpty() && ParseJsonObjectFromString(VideoMetadataString);
+}
+
+bool FLiveLinkFaceMetadataParser::ParseVideoMetadataFromString(const FString& InJsonString, FLiveLinkFaceVideoMetadata& OutVideoMetadata)
+{
+	if (TSharedPtr<FJsonObject> VideoMetadataJson = ParseJsonObjectFromString(InJsonString))
 	{
 		// Sample video_metadata.json
 		// {
@@ -356,7 +365,9 @@ bool FLiveLinkFaceMetadataParser::ParseDepthMetadata(const FString& InTakeDirect
 {
 	const FString& DepthMetadataFile = InTakeDirectory / FLiveLinkFaceStaticFileNames::DepthMetadata;
 
-	if (TSharedPtr<FJsonObject> DepthMetadataJson = ParseJsonObjectFromFile(DepthMetadataFile))
+	const FString& DepthMetadataString = ParseJsonStringFromFile(DepthMetadataFile);
+
+	if (TSharedPtr<FJsonObject> DepthMetadataJson = ParseJsonObjectFromString(DepthMetadataString))
 	{
 		DepthMetadataJson->TryGetStringField(TEXT("Build"), OutDepthMetadata.Build);
 		DepthMetadataJson->TryGetStringField(TEXT("Compression"), OutDepthMetadata.Compression);
@@ -438,7 +449,9 @@ bool FLiveLinkFaceMetadataParser::ParseAudioMetadata(const FString& InTakeDirect
 {
 	const FString& AudioMetadataFile = InTakeDirectory / FLiveLinkFaceStaticFileNames::AudioMetadata;
 
-	if (TSharedPtr<FJsonObject> AudioMetadataJson = ParseJsonObjectFromFile(AudioMetadataFile))
+	const FString& AudioMetadataString = ParseJsonStringFromFile(AudioMetadataFile);
+
+	if (TSharedPtr<FJsonObject> AudioMetadataJson = ParseJsonObjectFromString(AudioMetadataString))
 	{
 		// Sample audio_metadata.json
 		// {
@@ -463,7 +476,8 @@ bool FLiveLinkFaceMetadataParser::ParseTakeInfo(const FString& InTakeDirectory,
 {
 	const FString TakeMetadataFile = InTakeDirectory / FLiveLinkFaceStaticFileNames::TakeMetadata;
 
-	const TSharedPtr TakeMetadataJson = ParseJsonObjectFromFile(TakeMetadataFile);
+	const FString& TakeMetadataString = ParseJsonStringFromFile(TakeMetadataFile);
+	const TSharedPtr TakeMetadataJson = ParseJsonObjectFromString(TakeMetadataString);
 
 	if (!TakeMetadataJson.IsValid())
 	{
@@ -629,27 +643,32 @@ TArray<TSharedPtr<FJsonValue>> FLiveLinkFaceMetadataParser::ParseJsonArrayFromFi
 	return TArray<TSharedPtr<FJsonValue>>();
 }
 
-TSharedPtr<FJsonObject> FLiveLinkFaceMetadataParser::ParseJsonObjectFromFile(const FString& InFilePath)
+const FString FLiveLinkFaceMetadataParser::ParseJsonStringFromFile(const FString& InFilePath)
 {
 	FString JsonStringBuffer;
 
 	if (!IFileManager::Get().FileExists(*InFilePath))
 	{
 		UE_LOG(LogLiveLinkFaceMetadata, Error, TEXT("File not found: %s"), *InFilePath);
-		return nullptr;
+		return {};
 	}
 
 	if (!FFileHelper::LoadFileToString(JsonStringBuffer, *InFilePath))
 	{
 		UE_LOG(LogLiveLinkFaceMetadata, Error, TEXT("Failed to load file (check permissions): %s"), *InFilePath);
-		return nullptr;
+		return {};
 	}
 
+	return JsonStringBuffer;
+}
+
+TSharedPtr<FJsonObject> FLiveLinkFaceMetadataParser::ParseJsonObjectFromString(const FString& InJsonString)
+{
 	TSharedPtr<FJsonObject> Result;
 
-	if (!FJsonSerializer::Deserialize(TJsonReaderFactory<>::Create(JsonStringBuffer), Result))
+	if (!FJsonSerializer::Deserialize(TJsonReaderFactory<>::Create(InJsonString), Result))
 	{
-		UE_LOG(LogLiveLinkFaceMetadata, Error, TEXT("Failed to load json file (check for syntax errors): %s"), *InFilePath);
+		UE_LOG(LogLiveLinkFaceMetadata, Error, TEXT("Failed to parse json string (check for syntax errors): %s"), *InJsonString);
 		return nullptr;
 	}
 
@@ -1035,6 +1054,14 @@ TOptional<FTakeMetadata> ParseOldLiveLinkTakeMetadata(const FString& InTakeDirec
 	return ConvertOldToNewTakeMetadata(TakeInfo);
 }
 
+LIVELINKFACEMETADATA_API TArray<FTakeMetadata::FVideo> ParseOldLiveLinkVideoMetadataFromString(const FString& InJsonString, TArray<FText>& OutValidationErrors)
+{
+	FLiveLinkFaceTakeInfo TakeInfo;
+	REPORT_IF_ERROR_AND_RETURN(FLiveLinkFaceMetadataParser::ParseVideoMetadataFromString(InJsonString, TakeInfo.VideoMetadata), LOCTEXT("ParseOldTakeMetadata_ParseVideoInfoFailed", "Failed to parse take video metadata"));
+
+	return ConvertOldToNewTakeMetadata(TakeInfo).Video;
+}
+
 }
 
 }
diff --git a/Engine/Plugins/VirtualProduction/CaptureManager/CaptureManagerApp/Source/LiveLinkFaceMetadata/Public/LiveLinkFaceMetadata.h b/Engine/Plugins/VirtualProduction/CaptureManager/CaptureManagerApp/Source/LiveLinkFaceMetadata/Public/LiveLinkFaceMetadata.h
--- a/Engine/Plugins/VirtualProduction/CaptureManager/CaptureManagerApp/Source/LiveLinkFaceMetadata/Public/LiveLinkFaceMetadata.h
+++ b/Engine/Plugins/VirtualProduction/CaptureManager/CaptureManagerApp/Source/LiveLinkFaceMetadata/Public/LiveLinkFaceMetadata.h
@@ -13,5 +13,6 @@ namespace UE::CaptureManager::LiveLinkMetadata
 {
 
 LIVELINKFACEMETADATA_API TOptional<FTakeMetadata> ParseOldLiveLinkTakeMetadata(const FString& InJsonFile, TArray<FText>& OutValidationError);
+LIVELINKFACEMETADATA_API TArray<FTakeMetadata::FVideo> ParseOldLiveLinkVideoMetadataFromString(const FString& InJsonString, TArray<FText>& OutValidationError);
 
 }
\ No newline at end of file
diff --git a/Engine/Plugins/VirtualProduction/CaptureManager/CaptureManagerDevices/Source/CPSLiveLinkDevice/Private/LiveLinkFaceDevice.cpp b/Engine/Plugins/VirtualProduction/CaptureManager/CaptureManagerDevices/Source/CPSLiveLinkDevice/Private/LiveLinkFaceDevice.cpp
--- a/Engine/Plugins/VirtualProduction/CaptureManager/CaptureManagerDevices/Source/CPSLiveLinkDevice/Private/LiveLinkFaceDevice.cpp
+++ b/Engine/Plugins/VirtualProduction/CaptureManager/CaptureManagerDevices/Source/CPSLiveLinkDevice/Private/LiveLinkFaceDevice.cpp
@@ -89,7 +89,7 @@ void ULiveLinkFaceDevice::UpdateTakeList_Implementation()
 		NameToIdMap.Add(Take.Name, TakeId);
 	}
 
-	FetchThumbnails(MoveTemp(NameToIdMap));
+	FetchPreIngestFiles(MoveTemp(NameToIdMap));
 
 	ExecuteUpdateTakeListCallback(Execute_GetTakeIdentifiers(this));
 }
@@ -316,7 +316,10 @@ void ULiveLinkFaceDevice::HandleCPSEvent(TSharedPtr<const UE::CaptureManager::FC
 
 			PublishEvent<FIngestCapability_TakeAddedEvent>(TakeId);
 
-			FetchThumbnail(TakeName, TakeId);
+			TMap<FString, UE::CaptureManager::FTakeId> NameToIdMap;
+			NameToIdMap.Add(TakeName, TakeId);
+
+			FetchPreIngestFiles(MoveTemp(NameToIdMap));
 		}
 	}
 	else if (Event->UpdateMessage->GetAddressPath() == CPS::AddressPaths::GTakeRemoved)
@@ -457,48 +460,64 @@ TOptional<FTakeMetadata> ULiveLinkFaceDevice::ParseTake(const FString& InTakeDir
 	return {};
 }
 
-void ULiveLinkFaceDevice::FetchThumbnail(const FString& InTakeName, UE::CaptureManager::FTakeId InTakeId)
-{
-	TMap<FString, UE::CaptureManager::FTakeId> NameToIdMap;
-	NameToIdMap.Add(InTakeName, InTakeId);
-
-	FetchThumbnails(MoveTemp(NameToIdMap));
-}
-
-void ULiveLinkFaceDevice::FetchThumbnails(TMap<FString, UE::CaptureManager::FTakeId> InNameToIdMap)
+void ULiveLinkFaceDevice::FetchPreIngestFiles(TMap<FString, UE::CaptureManager::FTakeId> InNameToIdMap)
 {
 	using namespace UE::CaptureManager;
 
 	FCPSDataStream::FFileExportFinished OnFileExportFinishedCallback = FCPSDataStream::FFileExportFinished::CreateLambda(
-		[this, NameToIdMap = MoveTemp(InNameToIdMap)](const FString& InTakeName, TProtocolResult<FCPSDataStream::FData> InData)
+		[this, NameToIdMap = MoveTemp(InNameToIdMap)](TMap<FString, TMap<FString, TProtocolResult<FCPSDataStream::FData>>> InData)
 		{
-			if (InData.HasError())
+			for (const TPair<FString, TMap<FString, TProtocolResult<FCPSDataStream::FData>>>& TakePair : InData)
 			{
-				return;
-			}
+				const FString& TakeName = TakePair.Key;
+				for (const TPair<FString, TProtocolResult<FCPSDataStream::FData>>& FilePair : TakePair.Value)
+				{
+					const FString& FileName = FilePair.Key;
+					TProtocolResult<FCPSDataStream::FData> Result = FilePair.Value;
 
-			FCPSDataStream::FData Data = InData.StealValue();
+					if (Result.HasError())
+					{
+						continue;
+					}
 
-			if (const FTakeId* TakeId = NameToIdMap.Find(InTakeName))
-			{
-				TOptional<FTakeMetadata> TakeMetadataOpt = GetTakeMetadata(*TakeId);
-				if (!TakeMetadataOpt.IsSet())
-				{
-					return;
-				}
+					if (const FTakeId* TakeId = NameToIdMap.Find(TakeName))
+					{
+						TOptional<FTakeMetadata> TakeMetadataOpt = GetTakeMetadata(*TakeId);
+						if (!TakeMetadataOpt.IsSet())
+						{
+							return;
+						}
+
+						FTakeMetadata TakeMetadata = TakeMetadataOpt.GetValue();
+
+						FCPSDataStream::FData Data = Result.GetValue();
 
-				FTakeMetadata TakeMetadata = TakeMetadataOpt.GetValue();
-				TakeMetadata.Thumbnail = FTakeThumbnailData(MoveTemp(Data));
+						if (FileName == "thumbnail.jpg")
+						{
+							TakeMetadata.Thumbnail = FTakeThumbnailData(MoveTemp(Data));
+						}
+						else if (FileName == "video_metadata.json")
+						{
 
-				UpdateTake(*TakeId, MoveTemp(TakeMetadata));
+							FString DataString = FString(StringCast<UTF8CHAR>(reinterpret_cast<const char*>(Data.GetData()), Data.Num()));
 
-				PublishEvent<FIngestCapability_TakeUpdatedEvent>(*TakeId);
+							TArray<FText> ValidationFailures;
+							TArray<FTakeMetadata::FVideo> ParsedVideoObject = LiveLinkMetadata::ParseOldLiveLinkVideoMetadataFromString(DataString, ValidationFailures);
+
+							TakeMetadata.Video = ParsedVideoObject;
+						}
+
+						UpdateTake(*TakeId, MoveTemp(TakeMetadata));
+
+						PublishEvent<FIngestCapability_TakeUpdatedEvent>(*TakeId);
+					}
+				}
 			}
 		});
 
 	TUniquePtr<FCPSDataStream> DataStream = MakeUnique<FCPSDataStream>(MoveTemp(OnFileExportFinishedCallback));
 
-	Device->FetchThumbnails(MoveTemp(DataStream));
+	Device->FetchFiles(MoveTemp(DataStream), { FString("thumbnail.jpg"), FString("video_metadata.json") });
 }
 
 void ULiveLinkFaceDevice::RemoveDownloadedTakeData(const UE::CaptureManager::FTakeId InTakeId)
diff --git a/Engine/Plugins/VirtualProduction/CaptureManager/CaptureManagerDevices/Source/CPSLiveLinkDevice/Private/LiveLinkFaceDevice.h b/Engine/Plugins/VirtualProduction/CaptureManager/CaptureManagerDevices/Source/CPSLiveLinkDevice/Private/LiveLinkFaceDevice.h
--- a/Engine/Plugins/VirtualProduction/CaptureManager/CaptureManagerDevices/Source/CPSLiveLinkDevice/Private/LiveLinkFaceDevice.h
+++ b/Engine/Plugins/VirtualProduction/CaptureManager/CaptureManagerDevices/Source/CPSLiveLinkDevice/Private/LiveLinkFaceDevice.h
@@ -99,8 +99,7 @@ private:
 	TOptional<FTakeMetadata> ParseTake(const FString& InTakeDirectory, const FString& InTakeName);
 	FTakeMetadata ParseTakeMetadata(const UE::CaptureManager::FGetTakeMetadataResponse::FTakeObject& InTake);
 
-	void FetchThumbnail(const FString& InTakeName, UE::CaptureManager::FTakeId InTakeId);
-	void FetchThumbnails(TMap<FString, UE::CaptureManager::FTakeId> InNameToIdMap);
+	void FetchPreIngestFiles(TMap<FString, UE::CaptureManager::FTakeId> InNameToIdMap);
 
 	void RemoveDownloadedTakeData(const UE::CaptureManager::FTakeId InTakeId);
 
diff --git a/Engine/Plugins/VirtualProduction/CaptureManager/CaptureManagerDevices/Source/CPSLiveLinkDevice/Private/Protocol/CPSDataStream.cpp b/Engine/Plugins/VirtualProduction/CaptureManager/CaptureManagerDevices/Source/CPSLiveLinkDevice/Private/Protocol/CPSDataStream.cpp
--- a/Engine/Plugins/VirtualProduction/CaptureManager/CaptureManagerDevices/Source/CPSLiveLinkDevice/Private/Protocol/CPSDataStream.cpp
+++ b/Engine/Plugins/VirtualProduction/CaptureManager/CaptureManagerDevices/Source/CPSLiveLinkDevice/Private/Protocol/CPSDataStream.cpp
@@ -36,16 +36,20 @@ bool FCPSDataStream::FinishFile(const FString& InTakeName, const FString& InFile
 
 	if (FMemory::Memcmp(Hash.GetBytes(), InHash.GetData(), Hash.GetSize()) != 0)
 	{
-		FileExportFinished.ExecuteIfBound(InTakeName, FCaptureProtocolError(TEXT("Invalid file hash")));
+		TMap<FString, UE::CaptureManager::TProtocolResult<FData>>& TakeEntry = ExportResults.FindOrAdd(InTakeName);
+		TakeEntry.Add(InFileName, FCaptureProtocolError(TEXT("Invalid file hash")));
+		return true;
 	}
 
-	FileExportFinished.ExecuteIfBound(InTakeName, MoveTemp(Data));
+	TMap<FString, UE::CaptureManager::TProtocolResult<FData>>& TakeEntry = ExportResults.FindOrAdd(InTakeName);
+	TakeEntry.Add(InFileName, MoveTemp(Data));
 
 	return true;
 }
 
 void FCPSDataStream::Finalize(UE::CaptureManager::TProtocolResult<void> InResult)
 {
+	FileExportFinished.ExecuteIfBound(MoveTemp(ExportResults));
 }
 
 }
\ No newline at end of file
diff --git a/Engine/Plugins/VirtualProduction/CaptureManager/CaptureManagerDevices/Source/CPSLiveLinkDevice/Private/Protocol/CPSDataStream.h b/Engine/Plugins/VirtualProduction/CaptureManager/CaptureManagerDevices/Source/CPSLiveLinkDevice/Private/Protocol/CPSDataStream.h
--- a/Engine/Plugins/VirtualProduction/CaptureManager/CaptureManagerDevices/Source/CPSLiveLinkDevice/Private/Protocol/CPSDataStream.h
+++ b/Engine/Plugins/VirtualProduction/CaptureManager/CaptureManagerDevices/Source/CPSLiveLinkDevice/Private/Protocol/CPSDataStream.h
@@ -4,14 +4,17 @@
 
 #include "ExportClient/ExportClient.h"
 
+#include "Containers/Map.h"
+
 namespace UE::CaptureManager {
 
 class FCPSDataStream : public UE::CaptureManager::FBaseStream
 {
 public:
 	using FData = TArray<uint8>;
+	using FResults = TMap<FString, TMap<FString, UE::CaptureManager::TProtocolResult<FData>>>;
 
-	DECLARE_DELEGATE_TwoParams(FFileExportFinished, const FString& InTakeName, UE::CaptureManager::TProtocolResult<FData> InDataResult);
+	DECLARE_DELEGATE_OneParam(FFileExportFinished, FResults InExportResults);
 
 	FCPSDataStream(FFileExportFinished InFileExportFinished);
 
@@ -22,9 +25,10 @@ public:
 	void Finalize(UE::CaptureManager::TProtocolResult<void> InResult) override;
 
 private:
-
 	FData Data;
 	FFileExportFinished FileExportFinished;
+
+	FResults ExportResults;
 };
 
 }
\ No newline at end of file
diff --git a/Engine/Plugins/VirtualProduction/CaptureManager/CaptureManagerDevices/Source/CPSLiveLinkDevice/Private/Protocol/CPSDevice.cpp b/Engine/Plugins/VirtualProduction/CaptureManager/CaptureManagerDevices/Source/CPSLiveLinkDevice/Private/Protocol/CPSDevice.cpp
--- a/Engine/Plugins/VirtualProduction/CaptureManager/CaptureManagerDevices/Source/CPSLiveLinkDevice/Private/Protocol/CPSDevice.cpp
+++ b/Engine/Plugins/VirtualProduction/CaptureManager/CaptureManagerDevices/Source/CPSLiveLinkDevice/Private/Protocol/CPSDevice.cpp
@@ -266,10 +266,18 @@ void FCPSDevice::CancelExport(FTakeId InTakeId)
 	}
 }
 
-void FCPSDevice::FetchThumbnail(FTakeId InTakeId, TUniquePtr<FBaseStream> InStream)
+void FCPSDevice::FetchThumbnailForTake(FTakeId InTakeId, TUniquePtr<FBaseStream> InStream)
 {
-	static constexpr FStringView FileName = TEXTVIEW("thumbnail.jpg");
+	FetchFileForTake(InTakeId, MoveTemp(InStream), FString("thumbnail.jpg"));
+}
+
+void FCPSDevice::FetchThumbnails(TUniquePtr<FBaseStream> InStream)
+{
+	FetchFiles(MoveTemp(InStream), TArray<FString>{ FString("thumbnail.jpg") });
+}
 
+void FCPSDevice::FetchFileForTake(FTakeId InTakeId, TUniquePtr<FBaseStream> InStream, const FString& InFileName)
+{
 	FScopeLock Lock(&Mutex);
 	if (FGetTakeMetadataResponse::FTakeObject* TakeObjectPtr = TakeMetadata.Find(InTakeId))
 	{
@@ -282,7 +290,7 @@ void FCPSDevice::FetchThumbnail(FTakeId InTakeId, TUniquePtr<FBaseStream> InStre
 		uint64 Length = 0;
 		for (const FGetTakeMetadataResponse::FFileObject& FileObject : TakeObject.Files)
 		{
-			if (FileObject.Name == FileName)
+			if (FileObject.Name == InFileName)
 			{
 				Length = FileObject.Length;
 			}
@@ -294,7 +302,7 @@ void FCPSDevice::FetchThumbnail(FTakeId InTakeId, TUniquePtr<FBaseStream> InStre
 		}
 
 		FTakeFile TakeFile;
-		TakeFile.FileName = FileName;
+		TakeFile.FileName = InFileName;
 		TakeFile.Length = Length;
 		TakeFile.Offset = 0;
 
@@ -302,10 +310,8 @@ void FCPSDevice::FetchThumbnail(FTakeId InTakeId, TUniquePtr<FBaseStream> InStre
 	}
 }
 
-void FCPSDevice::FetchThumbnails(TUniquePtr<FBaseStream> InStream)
+void FCPSDevice::FetchFiles(TUniquePtr<FBaseStream> InStream, TArray<FString> InFileNames)
 {
-	static constexpr FStringView FileName = TEXTVIEW("thumbnail.jpg");
-
 	TMap<FString, FExportClient::FTakeFileArray> TakeFileArrayMap;
 
 	FScopeLock Lock(&Mutex);
@@ -315,25 +321,30 @@ void FCPSDevice::FetchThumbnails(TUniquePtr<FBaseStream> InStream)
 		const FGetTakeMetadataResponse::FTakeObject& TakeObject = TakePair.Value;
 
 		uint64 Length = 0;
+		FString FileName;
+		FExportClient::FTakeFileArray FileArray;
 		for (const FGetTakeMetadataResponse::FFileObject& FileObject : TakeObject.Files)
 		{
-			if (FileObject.Name == FileName)
+			if (InFileNames.Contains(FileObject.Name))
 			{
 				Length = FileObject.Length;
-			}
-		}
+				FileName = FileObject.Name;
 
-		if (Length == 0)
-		{
-			continue;
-		}
+				if (Length == 0)
+				{
+					continue;
+				}
 
-		FTakeFile TakeFile;
-		TakeFile.FileName = FileName;
-		TakeFile.Length = Length;
-		TakeFile.Offset = 0;
+				FTakeFile TakeFile;
+				TakeFile.FileName = FileName;
+				TakeFile.Length = Length;
+				TakeFile.Offset = 0;
+
+				FileArray.Add(MoveTemp(TakeFile));
+			}
+		}
 
-		TakeFileArrayMap.Add(TakeObject.Name, { MoveTemp(TakeFile) });
+		TakeFileArrayMap.Add(TakeObject.Name, MoveTemp(FileArray));
 	}
 
 	// Unlock earlier as there is no need to lock the ExporFiles function
diff --git a/Engine/Plugins/VirtualProduction/CaptureManager/CaptureManagerDevices/Source/CPSLiveLinkDevice/Private/Protocol/CPSDevice.h b/Engine/Plugins/VirtualProduction/CaptureManager/CaptureManagerDevices/Source/CPSLiveLinkDevice/Private/Protocol/CPSDevice.h
--- a/Engine/Plugins/VirtualProduction/CaptureManager/CaptureManagerDevices/Source/CPSLiveLinkDevice/Private/Protocol/CPSDevice.h
+++ b/Engine/Plugins/VirtualProduction/CaptureManager/CaptureManagerDevices/Source/CPSLiveLinkDevice/Private/Protocol/CPSDevice.h
@@ -96,9 +96,16 @@ public:
 	void StartExport(FTakeId InTakeId, TUniquePtr<FBaseStream> InStream);
 	void CancelExport(FTakeId InTakeId);
 
-	void FetchThumbnail(FTakeId InTakeId, TUniquePtr<FBaseStream> InStream);
+	/** Fetches the thumbnail for a single take. */
+	void FetchThumbnailForTake(FTakeId InTakeId, TUniquePtr<FBaseStream> InStream);
+	/** Fetches the thumbnails for all takes. */
 	void FetchThumbnails(TUniquePtr<FBaseStream> InStream);
 
+	/** Fetches the sepcified file for a single take. */
+	void FetchFileForTake(FTakeId InTakeId, TUniquePtr<FBaseStream> InStream, const FString& InFileName);
+	/** Fetches the specified files for all takes. */
+	void FetchFiles(TUniquePtr<FBaseStream> InStream, TArray<FString> InFileNames);
+
 private:
 
 	inline static const float ConnectInterval = 5.0f;
