[CaptureManager] Fix for live link face metadata rotation parsing

diff --git a/Engine/Plugins/VirtualProduction/CaptureManager/CaptureManagerApp/Source/LiveLinkFaceMetadata/Private/LiveLinkFaceMetadata.cpp b/Engine/Plugins/VirtualProduction/CaptureManager/CaptureManagerApp/Source/LiveLinkFaceMetadata/Private/LiveLinkFaceMetadata.cpp
--- a/Engine/Plugins/VirtualProduction/CaptureManager/CaptureManagerApp/Source/LiveLinkFaceMetadata/Private/LiveLinkFaceMetadata.cpp
+++ b/Engine/Plugins/VirtualProduction/CaptureManager/CaptureManagerApp/Source/LiveLinkFaceMetadata/Private/LiveLinkFaceMetadata.cpp
@@ -321,7 +321,7 @@ bool FLiveLinkFaceMetadataParser::ParseVideoMetadata(const FString& InTakeDirect
 	const FString& VideoMetadataFile = InTakeDirectory / FLiveLinkFaceStaticFileNames::VideoMetadata;
 	const FString& VideoMetadataString = ParseJsonStringFromFile(VideoMetadataFile);
 
-	return !VideoMetadataString.IsEmpty() && ParseJsonObjectFromString(VideoMetadataString);
+	return !VideoMetadataString.IsEmpty() && ParseVideoMetadataFromString(VideoMetadataString, OutVideoMetadata);
 }
 
 bool FLiveLinkFaceMetadataParser::ParseVideoMetadataFromString(const FString& InJsonString, FLiveLinkFaceVideoMetadata& OutVideoMetadata)
