[CaptureManager] Suggest enabling third party encoder settings if media read fails

diff --git a/Engine/Plugins/VirtualProduction/CaptureManager/CaptureManagerApp/Source/CaptureDataConverter/Private/Nodes/CaptureConvertAudioData.cpp b/Engine/Plugins/VirtualProduction/CaptureManager/CaptureManagerApp/Source/CaptureDataConverter/Private/Nodes/CaptureConvertAudioData.cpp
--- a/Engine/Plugins/VirtualProduction/CaptureManager/CaptureManagerApp/Source/CaptureDataConverter/Private/Nodes/CaptureConvertAudioData.cpp
+++ b/Engine/Plugins/VirtualProduction/CaptureManager/CaptureManagerApp/Source/CaptureDataConverter/Private/Nodes/CaptureConvertAudioData.cpp
@@ -85,7 +85,7 @@ FCaptureConvertAudioData::FResult FCaptureConvertAudioData::ConvertAudioFile()
 
 	if (AudioReaderResult.HasError())
 	{
-		FText Message = FText::Format(LOCTEXT("CaptureConvertAudioNode_UnsupportedFile", "Audio file format is unsupported {0}"), FText::FromString(AudioFilePath));
+		FText Message = FText::Format(LOCTEXT("CaptureConvertAudioNode_UnsupportedFile", "Audio file format is unsupported {0}. Consider enabling Third Party Encoder in Capture Manager settings."), FText::FromString(AudioFilePath));
 		return MakeError(MoveTemp(Message));
 	}
 
diff --git a/Engine/Plugins/VirtualProduction/CaptureManager/CaptureManagerApp/Source/CaptureDataConverter/Private/Nodes/CaptureConvertVideoData.cpp b/Engine/Plugins/VirtualProduction/CaptureManager/CaptureManagerApp/Source/CaptureDataConverter/Private/Nodes/CaptureConvertVideoData.cpp
--- a/Engine/Plugins/VirtualProduction/CaptureManager/CaptureManagerApp/Source/CaptureDataConverter/Private/Nodes/CaptureConvertVideoData.cpp
+++ b/Engine/Plugins/VirtualProduction/CaptureManager/CaptureManagerApp/Source/CaptureDataConverter/Private/Nodes/CaptureConvertVideoData.cpp
@@ -72,7 +72,7 @@ FCaptureConvertVideoData::FResult FCaptureConvertVideoData::ConvertData()
 
 	if (VideoReaderResult.HasError() || ImageWriterResult.HasError())
 	{
-		FText Message = FText::Format(LOCTEXT("CaptureConvertVideo_UnsupportedFile", "Video file format is unsupported {0}"), FText::FromString(VideoFilePath));
+		FText Message = FText::Format(LOCTEXT("CaptureConvertVideo_UnsupportedFile", "Video file format is unsupported {0}. Consider enabling Third Party Encoder in Capture Manager settings."), FText::FromString(VideoFilePath));
 		return MakeError(MoveTemp(Message));
 	}
 
diff --git a/Engine/Plugins/VirtualProduction/CaptureManager/CaptureManagerDevices/Source/MonoVideoLiveLinkDevice/Private/MonoVideoLiveLinkDevice.cpp b/Engine/Plugins/VirtualProduction/CaptureManager/CaptureManagerDevices/Source/MonoVideoLiveLinkDevice/Private/MonoVideoLiveLinkDevice.cpp
--- a/Engine/Plugins/VirtualProduction/CaptureManager/CaptureManagerDevices/Source/MonoVideoLiveLinkDevice/Private/MonoVideoLiveLinkDevice.cpp
+++ b/Engine/Plugins/VirtualProduction/CaptureManager/CaptureManagerDevices/Source/MonoVideoLiveLinkDevice/Private/MonoVideoLiveLinkDevice.cpp
@@ -278,7 +278,7 @@ TOptional<FTakeMetadata> UMonoVideoLiveLinkDevice::ReadTake(const FString& InCur
 	}
 	else
 	{
-		UE_LOG(LogVideoLiveLinkDevice, Warning, TEXT("Failed to determine the frame rate for the video file %s"), *InCurrentTakeFile);
+		UE_LOG(LogVideoLiveLinkDevice, Warning, TEXT("Failed to determine the frame rate for the video file %s. Consider enabling Third Party Encoder in Capture Manager settings."), *InCurrentTakeFile);
 	}
 
 	Video.Name = MoveTemp(Name);
diff --git a/Engine/Plugins/VirtualProduction/CaptureManager/CaptureManagerDevices/Source/StereoVideoLiveLinkDevice/Private/StereoVideoLiveLinkDevice.cpp b/Engine/Plugins/VirtualProduction/CaptureManager/CaptureManagerDevices/Source/StereoVideoLiveLinkDevice/Private/StereoVideoLiveLinkDevice.cpp
--- a/Engine/Plugins/VirtualProduction/CaptureManager/CaptureManagerDevices/Source/StereoVideoLiveLinkDevice/Private/StereoVideoLiveLinkDevice.cpp
+++ b/Engine/Plugins/VirtualProduction/CaptureManager/CaptureManagerDevices/Source/StereoVideoLiveLinkDevice/Private/StereoVideoLiveLinkDevice.cpp
@@ -496,7 +496,7 @@ TOptional<FTakeMetadata> UStereoVideoLiveLinkDevice::CreateTakeMetadata(const FT
 				}
 				else
 				{
-					UE_LOG(LogVideoLiveLinkDevice, Warning, TEXT("Failed to determine the frame rate for the video file %s"), *TakeComponent.Path);
+					UE_LOG(LogVideoLiveLinkDevice, Warning, TEXT("Failed to determine the frame rate for the video file %s. Consider enabling Third Party Encoder in Capture Manager settings."), *TakeComponent.Path);
 				}
 
 				Video.Path = TakeComponent.Path;
