Changes to the WMF usage

diff --git a/Engine/Plugins/VirtualProduction/CaptureManager/CaptureManagerApp/Source/CaptureDataConverter/Private/Nodes/CaptureConvertDepthData.cpp b/Engine/Plugins/VirtualProduction/CaptureManager/CaptureManagerApp/Source/CaptureDataConverter/Private/Nodes/CaptureConvertDepthData.cpp
--- a/Engine/Plugins/VirtualProduction/CaptureManager/CaptureManagerApp/Source/CaptureDataConverter/Private/Nodes/CaptureConvertDepthData.cpp
+++ b/Engine/Plugins/VirtualProduction/CaptureManager/CaptureManagerApp/Source/CaptureDataConverter/Private/Nodes/CaptureConvertDepthData.cpp
@@ -26,6 +26,11 @@ FCaptureConvertDepthData::FCaptureConvertDepthData(const FTakeMetadata::FVideo&
 {
 }
 
+FCaptureConvertDepthData::~FCaptureConvertDepthData()
+{
+	WritingQueue.Empty();
+}
+
 FCaptureConvertDepthData::FResult FCaptureConvertDepthData::Run()
 {
 	using namespace UE::CaptureManager;
@@ -98,12 +103,20 @@ FCaptureConvertDepthData::FResult FCaptureConvertDepthData::ConvertData()
 
 	TPromise<FResult> Promise;
 	TFuture<FResult> Future = Promise.GetFuture();
-	bShouldRun = true;
+	bIsDone = false;
+	bool bShouldRun = true;
 
 	while (bShouldRun)
 	{
 		TValueOrError<TUniquePtr<UE::CaptureManager::FMediaTextureSample>, FText> DepthSampleResult = DepthReader->Next();
 
+		// Stream ended
+		if (DepthSampleResult.HasValue() && !DepthSampleResult.GetValue().IsValid())
+		{
+			// Break when we set the nullptr value to the writer so it is aware of the ending
+			bShouldRun = false;
+		}
+
 		FWritingContext Context;
 		Context.ReadResult = MoveTemp(DepthSampleResult);
 		Context.Writer = ImageWriter.Get();
@@ -119,7 +132,10 @@ FCaptureConvertDepthData::FResult FCaptureConvertDepthData::ConvertData()
 		WritingQueue.Add(MoveTemp(Context));
 	}
 
-	return Future.Get();
+	FResult Result = Future.Get();
+	check(bIsDone); // Check that writing thread finished
+
+	return Result;
 }
 
 FCaptureConvertDepthData::FResult FCaptureConvertDepthData::CopyData()
@@ -198,7 +214,7 @@ FCaptureConvertDepthData::FResult FCaptureConvertDepthData::CopyData()
 
 void FCaptureConvertDepthData::OnWrite(FWritingContext InContext)
 {
-	if (!bShouldRun)
+	if (bIsDone)
 	{
 		return;
 	}
@@ -207,7 +223,7 @@ void FCaptureConvertDepthData::OnWrite(FWritingContext InContext)
 
 	if (InContext.ReadResult.HasError())
 	{
-		bShouldRun = false;
+		bIsDone = true;
 
 		DonePromise->SetValue(MakeError(InContext.ReadResult.StealError()));
 		return;
@@ -218,7 +234,7 @@ void FCaptureConvertDepthData::OnWrite(FWritingContext InContext)
 	// End of stream
 	if (!DepthSample.IsValid())
 	{
-		bShouldRun = false;
+		bIsDone = true;
 
 		DonePromise->SetValue(MakeValue());
 		return;
@@ -232,7 +248,7 @@ void FCaptureConvertDepthData::OnWrite(FWritingContext InContext)
 
 	if (WriterResult.IsSet())
 	{
-		bShouldRun = false;
+		bIsDone = true;
 
 		DonePromise->SetValue(MakeError(MoveTemp(WriterResult.GetValue())));
 
@@ -249,7 +265,7 @@ void FCaptureConvertDepthData::OnWrite(FWritingContext InContext)
 
 	if (Params.StopToken.IsStopRequested())
 	{
-		bShouldRun = false;
+		bIsDone = true;
 
 		FText Message = LOCTEXT("LiveLinkConvertDepthNode_AbortedByUser", "Aborted by user");
 		DonePromise->SetValue(MakeError(MoveTemp(Message)));
diff --git a/Engine/Plugins/VirtualProduction/CaptureManager/CaptureManagerApp/Source/CaptureDataConverter/Private/Nodes/CaptureConvertDepthData.h b/Engine/Plugins/VirtualProduction/CaptureManager/CaptureManagerApp/Source/CaptureDataConverter/Private/Nodes/CaptureConvertDepthData.h
--- a/Engine/Plugins/VirtualProduction/CaptureManager/CaptureManagerApp/Source/CaptureDataConverter/Private/Nodes/CaptureConvertDepthData.h
+++ b/Engine/Plugins/VirtualProduction/CaptureManager/CaptureManagerApp/Source/CaptureDataConverter/Private/Nodes/CaptureConvertDepthData.h
@@ -20,6 +20,8 @@ public:
 							 const FCaptureConvertDataNodeParams& InParams,
 							 const FCaptureConvertDepthOutputParams& InDepthParams);
 
+	virtual ~FCaptureConvertDepthData() override;
+
 private:
 
 	virtual FResult Run() override;
@@ -53,6 +55,6 @@ private:
 	void OnWrite(FWritingContext InContext);
 
 	UE::CaptureManager::TQueueRunner<FWritingContext> WritingQueue;
-	std::atomic_bool bShouldRun = false;
+	std::atomic_bool bIsDone = false;
 	int32 CurrentFrame = 0;
 };
\ No newline at end of file
diff --git a/Engine/Plugins/VirtualProduction/CaptureManager/CaptureManagerApp/Source/CaptureDataConverter/Private/Nodes/CaptureConvertVideoData.cpp b/Engine/Plugins/VirtualProduction/CaptureManager/CaptureManagerApp/Source/CaptureDataConverter/Private/Nodes/CaptureConvertVideoData.cpp
--- a/Engine/Plugins/VirtualProduction/CaptureManager/CaptureManagerApp/Source/CaptureDataConverter/Private/Nodes/CaptureConvertVideoData.cpp
+++ b/Engine/Plugins/VirtualProduction/CaptureManager/CaptureManagerApp/Source/CaptureDataConverter/Private/Nodes/CaptureConvertVideoData.cpp
@@ -21,11 +21,12 @@ FCaptureConvertVideoData::FCaptureConvertVideoData(const FTakeMetadata::FVideo&
 	: FConvertVideoNode(InVideo, InOutputDirectory)
 	, Params(InParams)
 	, VideoParams(InVideoParams)
-	, WritingQueue(UE::CaptureManager::TQueueRunner<FWritingContext>::FOnProcess::CreateRaw(this, &FCaptureConvertVideoData::OnWrite))
 {
 	checkf(!VideoParams.Format.IsEmpty(), TEXT("Output image format must be specified"));
 }
 
+FCaptureConvertVideoData::~FCaptureConvertVideoData() = default;
+
 FCaptureConvertVideoData::FResult FCaptureConvertVideoData::Run()
 {
 	using namespace UE::CaptureManager;
@@ -87,25 +88,41 @@ FCaptureConvertVideoData::FResult FCaptureConvertVideoData::ConvertData()
 
 	const double TotalDuration = VideoReader->GetDuration().GetTotalSeconds();
 
-	TPromise<FResult> Promise;
-	TFuture<FResult> Future = Promise.GetFuture();
-	bShouldRun = true;
+	FResult Result = MakeValue();
 
+	bool bShouldRun = true;
 	while (bShouldRun)
 	{
 		TValueOrError<TUniquePtr<UE::CaptureManager::FMediaTextureSample>, FText> VideoSampleResult = VideoReader->Next();
 
+		if (VideoSampleResult.HasError())
+		{
+			return MakeError(VideoSampleResult.StealError());
+		}
+
+		TUniquePtr<UE::CaptureManager::FMediaTextureSample> VideoSample = VideoSampleResult.StealValue();
+
+		// End of stream
+		if (!VideoSample.IsValid())
+		{
+			break;
+		}
+
 		FWritingContext Context;
-		Context.ReadResult = MoveTemp(VideoSampleResult);
+		Context.ReadSample = MoveTemp(VideoSample);
 		Context.Writer = ImageWriter.Get();
 		Context.TotalDuration = TotalDuration;
 		Context.Task = &Task;
-		Context.DonePromise = &Promise;
 
-		WritingQueue.Add(MoveTemp(Context));
+		FWriteResult WriteResult = OnWrite(MoveTemp(Context));
+
+		if (WriteResult.HasError())
+		{
+			return MakeError(WriteResult.StealError());
+		}
 	}
 
-	return Future.Get();
+	return Result;
 }
 
 FCaptureConvertVideoData::FResult FCaptureConvertVideoData::CopyData()
@@ -186,36 +203,12 @@ FCaptureConvertVideoData::FResult FCaptureConvertVideoData::CopyData()
 	return MakeValue();
 }
 
-void FCaptureConvertVideoData::OnWrite(FWritingContext InContext)
+FCaptureConvertVideoData::FWriteResult FCaptureConvertVideoData::OnWrite(FWritingContext InContext)
 {
-	if (!bShouldRun)
-	{
-		return;
-	}
-	
-	TPromise<FResult>* DonePromise = InContext.DonePromise;
-
-	if (InContext.ReadResult.HasError())
-	{
-		bShouldRun = false;
-
-		DonePromise->SetValue(MakeError(InContext.ReadResult.StealError()));
-		return;
-	}
-
-	TUniquePtr<UE::CaptureManager::FMediaTextureSample> VideoSample = InContext.ReadResult.StealValue();
-
-	// End of stream
-	if (!VideoSample.IsValid())
-	{
-		bShouldRun = false;
-		
-		DonePromise->SetValue(MakeValue());
-		return;
-	}
-
 	IImageWriter* Writer = InContext.Writer;
 
+	TUniquePtr<UE::CaptureManager::FMediaTextureSample>& VideoSample = InContext.ReadSample;
+
 	EMediaOrientation Orientation = ConvertOrientation(Video.Orientation.Get(FTakeMetadata::FVideo::EOrientation::Original));
 	VideoSample->Orientation = Orientation;
 	VideoSample->Rotation = VideoParams.Rotation;
@@ -225,11 +218,7 @@ void FCaptureConvertVideoData::OnWrite(FWritingContext InContext)
 
 	if (WriterResult.IsSet())
 	{
-		bShouldRun = false;
-
-		DonePromise->SetValue(MakeError(MoveTemp(WriterResult.GetValue())));
-
-		return;
+		return MakeError(MoveTemp(WriterResult.GetValue()));
 	}
 
 	const double Time = VideoSample->Time.GetTotalSeconds();
@@ -240,11 +229,11 @@ void FCaptureConvertVideoData::OnWrite(FWritingContext InContext)
 
 	if (Params.StopToken.IsStopRequested())
 	{
-		bShouldRun = false;
-
 		FText Message = LOCTEXT("CaptureConvertVideo_AbortedByUser", "Video conversion aborted by user");
-		DonePromise->SetValue(MakeError(MoveTemp(Message)));
+		return MakeError(MoveTemp(Message));
 	}
+
+	return MakeValue();
 }
 
 EMediaOrientation FCaptureConvertVideoData::ConvertOrientation(FTakeMetadata::FVideo::EOrientation InOrientation) const
diff --git a/Engine/Plugins/VirtualProduction/CaptureManager/CaptureManagerApp/Source/CaptureDataConverter/Private/Nodes/CaptureConvertVideoData.h b/Engine/Plugins/VirtualProduction/CaptureManager/CaptureManagerApp/Source/CaptureDataConverter/Private/Nodes/CaptureConvertVideoData.h
--- a/Engine/Plugins/VirtualProduction/CaptureManager/CaptureManagerApp/Source/CaptureDataConverter/Private/Nodes/CaptureConvertVideoData.h
+++ b/Engine/Plugins/VirtualProduction/CaptureManager/CaptureManagerApp/Source/CaptureDataConverter/Private/Nodes/CaptureConvertVideoData.h
@@ -21,6 +21,8 @@ public:
 							 const FCaptureConvertDataNodeParams& InParams,
 							 const FCaptureConvertVideoOutputParams& InVideoParams);
 
+	virtual ~FCaptureConvertVideoData() override;
+
 private:
 
 	virtual FResult Run() override;
@@ -36,25 +38,20 @@ private:
 	struct FWritingContext
 	{
 		FWritingContext() 
-			: ReadResult(MakeValue())
+			: ReadSample(nullptr)
 			, Writer(nullptr)
 			, Task(nullptr)
 			, TotalDuration(0.0f)
-			, DonePromise(nullptr)
 		{
 		}
 
-		TValueOrError<TUniquePtr<UE::CaptureManager::FMediaTextureSample>, FText> ReadResult;
+		TUniquePtr<UE::CaptureManager::FMediaTextureSample> ReadSample;
 		class IImageWriter* Writer;
 		UE::CaptureManager::FTaskProgress::FTask* Task;
 
 		double TotalDuration;
-
-		TPromise<FResult>* DonePromise;
 	};
 
-	void OnWrite(FWritingContext InContext);
-
-	UE::CaptureManager::TQueueRunner<FWritingContext> WritingQueue;
-	std::atomic_bool bShouldRun = false;
+	using FWriteResult = TValueOrError<void, FText>;
+	FWriteResult OnWrite(FWritingContext InContext);
 };
\ No newline at end of file
diff --git a/Engine/Plugins/VirtualProduction/CaptureManager/CaptureManagerApp/Source/CaptureManagerMediaRW/Private/MediaSample.cpp b/Engine/Plugins/VirtualProduction/CaptureManager/CaptureManagerApp/Source/CaptureManagerMediaRW/Private/MediaSample.cpp
--- a/Engine/Plugins/VirtualProduction/CaptureManager/CaptureManagerApp/Source/CaptureManagerMediaRW/Private/MediaSample.cpp
+++ b/Engine/Plugins/VirtualProduction/CaptureManager/CaptureManagerApp/Source/CaptureManagerMediaRW/Private/MediaSample.cpp
@@ -64,6 +64,7 @@ uint32 GetNumberOfChannels(EMediaTexturePixelFormat InPixelFormat)
 	{
 		case UE::CaptureManager::EMediaTexturePixelFormat::U8_RGBA:
 		case UE::CaptureManager::EMediaTexturePixelFormat::U8_BGRA:
+		case UE::CaptureManager::EMediaTexturePixelFormat::F_Mono:
 			return 4;
 		case UE::CaptureManager::EMediaTexturePixelFormat::U8_RGB:
 		case UE::CaptureManager::EMediaTexturePixelFormat::U8_BGR:
@@ -71,6 +72,8 @@ uint32 GetNumberOfChannels(EMediaTexturePixelFormat InPixelFormat)
 		case UE::CaptureManager::EMediaTexturePixelFormat::U8_I420:
 		case UE::CaptureManager::EMediaTexturePixelFormat::U8_NV12:
 			return 3;
+		case UE::CaptureManager::EMediaTexturePixelFormat::U16_Mono:
+			return 2;
 		case UE::CaptureManager::EMediaTexturePixelFormat::U8_Mono:
 		default:
 			return 1;
diff --git a/Engine/Plugins/VirtualProduction/CaptureManager/CaptureManagerApp/Source/CaptureManagerMediaRW/Private/Readers/WindowsMediaReader.cpp b/Engine/Plugins/VirtualProduction/CaptureManager/CaptureManagerApp/Source/CaptureManagerMediaRW/Private/Readers/WindowsMediaReader.cpp
--- a/Engine/Plugins/VirtualProduction/CaptureManager/CaptureManagerApp/Source/CaptureManagerMediaRW/Private/Readers/WindowsMediaReader.cpp
+++ b/Engine/Plugins/VirtualProduction/CaptureManager/CaptureManagerApp/Source/CaptureManagerMediaRW/Private/Readers/WindowsMediaReader.cpp
@@ -109,6 +109,8 @@ TOptional<FText> FWindowsAudioReader::Open(const FString& InFileName)
 
 TOptional<FText> FWindowsAudioReader::Close()
 {
+	AudioReader->Flush(MF_SOURCE_READER_FIRST_AUDIO_STREAM);
+
 	AudioReader.Reset();
 
 	return {};
@@ -300,6 +302,8 @@ TOptional<FText> FWindowsVideoReader::Open(const FString& InFileName)
 
 TOptional<FText> FWindowsVideoReader::Close()
 {
+	VideoReader->Flush(MF_SOURCE_READER_FIRST_VIDEO_STREAM);
+
 	VideoReader.Reset();
 
 	return {};
diff --git a/Engine/Plugins/VirtualProduction/CaptureManager/CaptureManagerApp/Source/CaptureManagerMediaRW/Private/Utils/WindowsJpegDecoder.cpp b/Engine/Plugins/VirtualProduction/CaptureManager/CaptureManagerApp/Source/CaptureManagerMediaRW/Private/Utils/WindowsJpegDecoder.cpp
--- a/Engine/Plugins/VirtualProduction/CaptureManager/CaptureManagerApp/Source/CaptureManagerMediaRW/Private/Utils/WindowsJpegDecoder.cpp
+++ b/Engine/Plugins/VirtualProduction/CaptureManager/CaptureManagerApp/Source/CaptureManagerMediaRW/Private/Utils/WindowsJpegDecoder.cpp
@@ -63,33 +63,26 @@ TOptional<FText> FWindowsJpegDecoder::Decode(uint8* InData, uint32 InSize, TArra
 	Result = Decoder->GetFrame(FrameIndex, &DecodedFrame);
 	WINJD_CHECK_AND_RETURN_ERROR_MESSAGE(Result, LOCTEXT("Decode_FailedToDecode", "Failed to decode the jpeg frame"));
 
-	TComPtr<IWICBitmap> Bitmap;
-	Result = WindowsImagingFactory->CreateBitmapFromSource(DecodedFrame, WICBitmapCacheOnDemand, &Bitmap);
-	WINJD_CHECK_AND_RETURN_ERROR_MESSAGE(Result, LOCTEXT("Decode_FailedCreateBitmap", "Failed to create the bitmap frame"));
-
 	uint32 Width = 0;
 	uint32 Height = 0;
 
 	DecodedFrame->GetSize(&Width, &Height);
 	WICRect Lock = { 0, 0, Width, Height };
-	TComPtr<IWICBitmapLock> BitmapLock;
-	Result = Bitmap->Lock(&Lock, WICBitmapLockWrite, &BitmapLock);
-	WINJD_CHECK_AND_RETURN_ERROR_MESSAGE(Result, LOCTEXT("Decode_FailedCreateLock", "Failed to create the bitmap lock to obtain the raw data"));
-
-	uint32 Size;
-	uint8* Buffer;
-	Result = BitmapLock->GetDataPointer(&Size, &Buffer);
-	WINJD_CHECK_AND_RETURN_ERROR_MESSAGE(Result, LOCTEXT("Decode_FailedObtainData", "Failed to obtain raw data"));
-
-	OutImage.SetNum(Size);
-	FMemory::Memcpy(OutImage.GetData(), Buffer, Size);
 
 	WICPixelFormatGUID PixelFormat;
-	Result = Bitmap->GetPixelFormat(&PixelFormat);
+	Result = DecodedFrame->GetPixelFormat(&PixelFormat);
 	WINJD_CHECK_AND_RETURN_ERROR_MESSAGE(Result, LOCTEXT("Decode_FailedObtainPixelFormat", "Failed to obtain the pixel format"));
 
 	OutPixelFormat = ConvertPixelFormat(PixelFormat);
 
+	uint32 BytesPerPixel = GetNumberOfChannels(OutPixelFormat);
+	uint32 Stride = (Width * BytesPerPixel);
+
+	OutImage.SetNum(Stride * Height);
+
+	Result = DecodedFrame->CopyPixels(&Lock, Stride, OutImage.Num(), OutImage.GetData());
+	WINJD_CHECK_AND_RETURN_ERROR_MESSAGE(Result, LOCTEXT("Decode_FailedCopyPixels", "Failed to obtain the pixels from the decoded image"));
+
 	return {};
 }
 
diff --git a/Engine/Plugins/VirtualProduction/CaptureManager/CaptureManagerApp/Source/CaptureManagerMediaRW/Private/Utils/WindowsJpegDecoder.h b/Engine/Plugins/VirtualProduction/CaptureManager/CaptureManagerApp/Source/CaptureManagerMediaRW/Private/Utils/WindowsJpegDecoder.h
--- a/Engine/Plugins/VirtualProduction/CaptureManager/CaptureManagerApp/Source/CaptureManagerMediaRW/Private/Utils/WindowsJpegDecoder.h
+++ b/Engine/Plugins/VirtualProduction/CaptureManager/CaptureManagerApp/Source/CaptureManagerMediaRW/Private/Utils/WindowsJpegDecoder.h
@@ -37,7 +37,7 @@ private:
 	FWindowsJpegDecoder();
 	TOptional<FText> Initialize();
 
-	UE::CaptureManager::EMediaTexturePixelFormat ConvertPixelFormat(WICPixelFormatGUID InPixelFormat);
+	static UE::CaptureManager::EMediaTexturePixelFormat ConvertPixelFormat(WICPixelFormatGUID InPixelFormat);
 
 	TComPtr<IWICImagingFactory> WindowsImagingFactory;
 };
diff --git a/Engine/Plugins/VirtualProduction/CaptureManager/CaptureManagerApp/Source/CaptureManagerMediaRW/Private/Writers/WindowsImageWriter.cpp b/Engine/Plugins/VirtualProduction/CaptureManager/CaptureManagerApp/Source/CaptureManagerMediaRW/Private/Writers/WindowsImageWriter.cpp
--- a/Engine/Plugins/VirtualProduction/CaptureManager/CaptureManagerApp/Source/CaptureManagerMediaRW/Private/Writers/WindowsImageWriter.cpp
+++ b/Engine/Plugins/VirtualProduction/CaptureManager/CaptureManagerApp/Source/CaptureManagerMediaRW/Private/Writers/WindowsImageWriter.cpp
@@ -330,6 +330,8 @@ TOptional<FText> FWindowsImageWriter::Append(UE::CaptureManager::FMediaTextureSa
 																   &Bitmap);
 	WINIW_CHECK_AND_RETURN_ERROR_MESSAGE(Result, LOCTEXT("Append_FailedToCreateBitmap", "Failed to create bitmap from memory"));
 
+	InSample->Buffer.Empty();
+
 	Private::FConverter Converter = Private::FConverter::Create(WindowsImagingFactory, Bitmap, InSample->CurrentFormat, InSample->DesiredFormat);
 
 	TComPtr<IWICBitmapSource> Source = Converter.GetSource();
diff --git a/Engine/Plugins/VirtualProduction/CaptureManager/CaptureManagerApp/Source/CaptureManagerPipeline/Private/Nodes/ConvertAudioNode.cpp b/Engine/Plugins/VirtualProduction/CaptureManager/CaptureManagerApp/Source/CaptureManagerPipeline/Private/Nodes/ConvertAudioNode.cpp
--- a/Engine/Plugins/VirtualProduction/CaptureManager/CaptureManagerApp/Source/CaptureManagerPipeline/Private/Nodes/ConvertAudioNode.cpp
+++ b/Engine/Plugins/VirtualProduction/CaptureManager/CaptureManagerApp/Source/CaptureManagerPipeline/Private/Nodes/ConvertAudioNode.cpp
@@ -26,6 +26,8 @@ FConvertAudioNode::FConvertAudioNode(const FTakeMetadata::FAudio& InAudio,
 	OutputDirectory = InOutputDirectory / AudioDirectory;
 }
 
+FConvertAudioNode::~FConvertAudioNode() = default;
+
 FConvertAudioNode::FResult FConvertAudioNode::Prepare()
 {
 	IFileManager& FileManager = IFileManager::Get();
diff --git a/Engine/Plugins/VirtualProduction/CaptureManager/CaptureManagerApp/Source/CaptureManagerPipeline/Private/Nodes/ConvertDepthNode.cpp b/Engine/Plugins/VirtualProduction/CaptureManager/CaptureManagerApp/Source/CaptureManagerPipeline/Private/Nodes/ConvertDepthNode.cpp
--- a/Engine/Plugins/VirtualProduction/CaptureManager/CaptureManagerApp/Source/CaptureManagerPipeline/Private/Nodes/ConvertDepthNode.cpp
+++ b/Engine/Plugins/VirtualProduction/CaptureManager/CaptureManagerApp/Source/CaptureManagerPipeline/Private/Nodes/ConvertDepthNode.cpp
@@ -23,6 +23,8 @@ FConvertDepthNode::FConvertDepthNode(const FTakeMetadata::FVideo& InDepth,
 	OutputDirectory = InOutputDirectory / DepthDirectory;
 }
 
+FConvertDepthNode::~FConvertDepthNode() = default;
+
 FConvertDepthNode::FResult FConvertDepthNode::Prepare()
 {
 	IFileManager& FileManager = IFileManager::Get();
diff --git a/Engine/Plugins/VirtualProduction/CaptureManager/CaptureManagerApp/Source/CaptureManagerPipeline/Private/Nodes/ConvertVideoNode.cpp b/Engine/Plugins/VirtualProduction/CaptureManager/CaptureManagerApp/Source/CaptureManagerPipeline/Private/Nodes/ConvertVideoNode.cpp
--- a/Engine/Plugins/VirtualProduction/CaptureManager/CaptureManagerApp/Source/CaptureManagerPipeline/Private/Nodes/ConvertVideoNode.cpp
+++ b/Engine/Plugins/VirtualProduction/CaptureManager/CaptureManagerApp/Source/CaptureManagerPipeline/Private/Nodes/ConvertVideoNode.cpp
@@ -23,6 +23,8 @@ FConvertVideoNode::FConvertVideoNode(const FTakeMetadata::FVideo& InVideo,
 	OutputDirectory = InOutputDirectory / VideoDirectory;
 }
 
+FConvertVideoNode::~FConvertVideoNode() = default;
+
 FConvertVideoNode::FResult FConvertVideoNode::Prepare()
 {
 	IFileManager& FileManager = IFileManager::Get();
diff --git a/Engine/Plugins/VirtualProduction/CaptureManager/CaptureManagerApp/Source/CaptureManagerPipeline/Public/Nodes/ConvertAudioNode.h b/Engine/Plugins/VirtualProduction/CaptureManager/CaptureManagerApp/Source/CaptureManagerPipeline/Public/Nodes/ConvertAudioNode.h
--- a/Engine/Plugins/VirtualProduction/CaptureManager/CaptureManagerApp/Source/CaptureManagerPipeline/Public/Nodes/ConvertAudioNode.h
+++ b/Engine/Plugins/VirtualProduction/CaptureManager/CaptureManagerApp/Source/CaptureManagerPipeline/Public/Nodes/ConvertAudioNode.h
@@ -13,6 +13,8 @@ public:
 	FConvertAudioNode(const FTakeMetadata::FAudio& InAudio,
 					  const FString& InOutputDirectory);
 
+	virtual ~FConvertAudioNode() override;
+
 protected:
 
 	FTakeMetadata::FAudio Audio;
diff --git a/Engine/Plugins/VirtualProduction/CaptureManager/CaptureManagerApp/Source/CaptureManagerPipeline/Public/Nodes/ConvertDepthNode.h b/Engine/Plugins/VirtualProduction/CaptureManager/CaptureManagerApp/Source/CaptureManagerPipeline/Public/Nodes/ConvertDepthNode.h
--- a/Engine/Plugins/VirtualProduction/CaptureManager/CaptureManagerApp/Source/CaptureManagerPipeline/Public/Nodes/ConvertDepthNode.h
+++ b/Engine/Plugins/VirtualProduction/CaptureManager/CaptureManagerApp/Source/CaptureManagerPipeline/Public/Nodes/ConvertDepthNode.h
@@ -12,6 +12,9 @@ public:
 
 	FConvertDepthNode(const FTakeMetadata::FVideo& InDepth,
 					  const FString& InOutputDirectory);
+
+	virtual ~FConvertDepthNode() override;
+
 protected:
 
 	FTakeMetadata::FVideo Depth;
diff --git a/Engine/Plugins/VirtualProduction/CaptureManager/CaptureManagerApp/Source/CaptureManagerPipeline/Public/Nodes/ConvertVideoNode.h b/Engine/Plugins/VirtualProduction/CaptureManager/CaptureManagerApp/Source/CaptureManagerPipeline/Public/Nodes/ConvertVideoNode.h
--- a/Engine/Plugins/VirtualProduction/CaptureManager/CaptureManagerApp/Source/CaptureManagerPipeline/Public/Nodes/ConvertVideoNode.h
+++ b/Engine/Plugins/VirtualProduction/CaptureManager/CaptureManagerApp/Source/CaptureManagerPipeline/Public/Nodes/ConvertVideoNode.h
@@ -13,6 +13,8 @@ public:
 	FConvertVideoNode(const FTakeMetadata::FVideo& InVideo, 
 					  const FString& InOutputDirectory);
 
+	virtual ~FConvertVideoNode() override;
+
 protected:
 
 	FTakeMetadata::FVideo Video;
