Fixing example python scripts

* Capture Manager scripts
* Depth Generator scripts
* Adhering to the MHA script argument format

diff --git a/Engine/Plugins/VirtualProduction/CaptureManager/CaptureManagerApp/Source/CaptureManagerEditor/Private/Scripting/ULiveLinkHubCaptureDevice.cpp b/Engine/Plugins/VirtualProduction/CaptureManager/CaptureManagerApp/Source/CaptureManagerEditor/Private/Scripting/ULiveLinkHubCaptureDevice.cpp
--- a/Engine/Plugins/VirtualProduction/CaptureManager/CaptureManagerApp/Source/CaptureManagerEditor/Private/Scripting/ULiveLinkHubCaptureDevice.cpp
+++ b/Engine/Plugins/VirtualProduction/CaptureManager/CaptureManagerApp/Source/CaptureManagerEditor/Private/Scripting/ULiveLinkHubCaptureDevice.cpp
@@ -294,7 +294,7 @@ ULiveLinkHubCaptureDevice* ULiveLinkHubCaptureDeviceFactory::CreateDeviceByClass
 	ULiveLinkDeviceSubsystem* DeviceSubsystem = GEngine->GetEngineSubsystem<ULiveLinkDeviceSubsystem>();
 	check(DeviceSubsystem);
 
-	ULiveLinkDeviceSubsystem::FCreateResult DeviceResult = DeviceSubsystem->CreateDeviceOfClass(DeviceClass);
+	ULiveLinkDeviceSubsystem::FCreateResult DeviceResult = DeviceSubsystem->CreateDeviceOfClass(DeviceClass, InSettings);
 	ULiveLinkDevice* NewDevice = DeviceResult.GetValue().Device;
 
 	check(NewDevice->Implements<ULiveLinkDeviceCapability_Ingest>());
diff --git a/Engine/Plugins/VirtualProduction/CaptureManager/CaptureManagerApp/Source/LiveLinkCapabilities/Public/Ingest/IngestCapability_Options.h b/Engine/Plugins/VirtualProduction/CaptureManager/CaptureManagerApp/Source/LiveLinkCapabilities/Public/Ingest/IngestCapability_Options.h
--- a/Engine/Plugins/VirtualProduction/CaptureManager/CaptureManagerApp/Source/LiveLinkCapabilities/Public/Ingest/IngestCapability_Options.h
+++ b/Engine/Plugins/VirtualProduction/CaptureManager/CaptureManagerApp/Source/LiveLinkCapabilities/Public/Ingest/IngestCapability_Options.h
@@ -47,10 +47,10 @@ struct LIVELINKCAPABILITIES_API FIngestCapability_VideoOptions
 	FString Format;
 
 	UPROPERTY(EditAnywhere, BlueprintReadWrite, Category = "Live Link Device|Ingest")
-	EIngestCapability_ImagePixelFormat PixelFormat;
+	EIngestCapability_ImagePixelFormat PixelFormat = EIngestCapability_ImagePixelFormat::U8_BGRA;
 
 	UPROPERTY(EditAnywhere, BlueprintReadWrite, Category = "Live Link Device|Ingest")
-	EIngestCapability_ImageRotation Rotation;
+	EIngestCapability_ImageRotation Rotation = EIngestCapability_ImageRotation::None;
 };
 
 USTRUCT(BlueprintType)
