Fixing warning on static analysis

diff --git a/Engine/Plugins/VirtualProduction/CaptureManager/CaptureManagerApp/Source/CaptureDataConverter/Private/Nodes/CaptureConvertVideoData.cpp b/Engine/Plugins/VirtualProduction/CaptureManager/CaptureManagerApp/Source/CaptureDataConverter/Private/Nodes/CaptureConvertVideoData.cpp
--- a/Engine/Plugins/VirtualProduction/CaptureManager/CaptureManagerApp/Source/CaptureDataConverter/Private/Nodes/CaptureConvertVideoData.cpp
+++ b/Engine/Plugins/VirtualProduction/CaptureManager/CaptureManagerApp/Source/CaptureDataConverter/Private/Nodes/CaptureConvertVideoData.cpp
@@ -90,8 +90,7 @@ FCaptureConvertVideoData::FResult FCaptureConvertVideoData::ConvertData()
 
 	FResult Result = MakeValue();
 
-	bool bShouldRun = true;
-	while (bShouldRun)
+	while (true)
 	{
 		TValueOrError<TUniquePtr<UE::CaptureManager::FMediaTextureSample>, FText> VideoSampleResult = VideoReader->Next();
 
