[CaptureManager] Fixed bug where add to queue button can become permanently disabled

diff --git a/Engine/Plugins/VirtualProduction/CaptureManager/CaptureManagerApp/Source/CaptureManagerEditor/Private/UI/SIngestJobProcessor.cpp b/Engine/Plugins/VirtualProduction/CaptureManager/CaptureManagerApp/Source/CaptureManagerEditor/Private/UI/SIngestJobProcessor.cpp
--- a/Engine/Plugins/VirtualProduction/CaptureManager/CaptureManagerApp/Source/CaptureManagerEditor/Private/UI/SIngestJobProcessor.cpp
+++ b/Engine/Plugins/VirtualProduction/CaptureManager/CaptureManagerApp/Source/CaptureManagerEditor/Private/UI/SIngestJobProcessor.cpp
@@ -28,6 +28,31 @@ static const FText TakeText = LOCTEXT("TakeText", "Take");
 static const FText ProgressText = LOCTEXT("ProgressText", "Progress");
 static const FText StatusText = LOCTEXT("StatusText", "Status");
 
+namespace Private
+{
+void ResetJob(TSharedPtr<FIngestJob> InIngestJob, TSharedPtr<FIngestJobProcessor> InIngestJobProcessor)
+{
+	if (InIngestJobProcessor->RemoveJobs({ InIngestJob->GetGuid(), }))
+	{
+		TSharedRef<FIngestJob> NewJob = MakeShared<FIngestJob>(
+			InIngestJob->GetCaptureDeviceId(),
+			InIngestJob->GetTakeId(),
+			InIngestJob->GetTakeMetadata(),
+			InIngestJob->GetSettings()
+		);
+
+		if (!InIngestJobProcessor->AddJobs({ MoveTemp(NewJob), }))
+		{
+			UE_LOG(LogTemp, Warning, TEXT("Failed to add new job entry"));
+		}
+	}
+	else
+	{
+		UE_LOG(LogTemp, Warning, TEXT("Failed to remove old job entry"));
+	}
+}
+}
+
 class SIngestJobRow : public SMultiColumnTableRow<TSharedRef<FIngestJob>>
 {
 public:
@@ -280,7 +305,19 @@ bool SIngestJobProcessor::StartButtonIsEnabled() const
 
 FReply SIngestJobProcessor::OnStartButtonClicked()
 {
-	IngestJobProcessor->StartProcessing();
+	for (const TSharedRef<FIngestJob>& IngestJob : IngestJobs)
+	{
+		if (IngestJob->GetProcessingState() == FIngestJob::EProcessingState::Aborted)
+		{
+			Private::ResetJob(IngestJob, IngestJobProcessor);
+		}
+	}
+
+	if (!IngestJobs.IsEmpty())
+	{
+		IngestJobProcessor->StartProcessing();
+	}
+
 	return FReply::Handled();
 }
 
@@ -461,24 +498,7 @@ TSharedRef<SWidget> SIngestJobRow::CreateActionWidget()
 
 FReply SIngestJobRow::RetryButtonClicked()
 {
-	if (IngestJobProcessor->RemoveJobs({ IngestJob->GetGuid(), }))
-	{
-		TSharedRef<FIngestJob> NewJob = MakeShared<FIngestJob>(
-			IngestJob->GetCaptureDeviceId(),
-			IngestJob->GetTakeId(),
-			IngestJob->GetTakeMetadata(),
-			IngestJob->GetSettings()
-		);
-
-		if (!IngestJobProcessor->AddJobs({ MoveTemp(NewJob), }))
-		{
-			UE_LOG(LogTemp, Warning, TEXT("Failed to add new job entry"));
-		}
-	}
-	else
-	{
-		UE_LOG(LogTemp, Warning, TEXT("Failed to remove old job entry"));
-	}
+	Private::ResetJob(IngestJob, IngestJobProcessor);
 
 	return FReply::Handled();
 }
