Fix crash on app exit while ingest in progress.
This resolves a portion of potential cases in which app can end up.
Additional testing required to identify other corner cases.
Redesign of the device system is also required to address this issue properly.

diff --git a/Engine/Plugins/VirtualProduction/CaptureManager/CaptureManagerApp/Source/CaptureManagerEditor/Private/IngestManagement/IngestJobProcessor.cpp b/Engine/Plugins/VirtualProduction/CaptureManager/CaptureManagerApp/Source/CaptureManagerEditor/Private/IngestManagement/IngestJobProcessor.cpp
--- a/Engine/Plugins/VirtualProduction/CaptureManager/CaptureManagerApp/Source/CaptureManagerEditor/Private/IngestManagement/IngestJobProcessor.cpp
+++ b/Engine/Plugins/VirtualProduction/CaptureManager/CaptureManagerApp/Source/CaptureManagerEditor/Private/IngestManagement/IngestJobProcessor.cpp
@@ -15,29 +15,41 @@ namespace UE::CaptureManager
 namespace Private
 {
 
-static void StopUnrealEndpointConnections()
-{
-	TSharedRef<FUnrealEndpointManager> UnrealEndpointManager = FModuleManager::LoadModuleChecked<FCaptureManagerUnrealEndpointModule>("CaptureManagerUnrealEndpoint").GetEndpointManager();
-	const TArray<TSharedRef<FUnrealEndpoint>> Endpoints = UnrealEndpointManager->GetEndpoints();
-
-	for (const TSharedRef<FUnrealEndpoint>& Endpoint : Endpoints)
+	static void StopUnrealEndpointConnections()
 	{
-		Endpoint->StopConnection();
-	}
+		FCaptureManagerUnrealEndpointModule* EndpointModule = FModuleManager::GetModulePtr<FCaptureManagerUnrealEndpointModule>("Settings");
+		if (!EndpointModule)
+		{
+			return;
+		}
+		
+		TOptional<TSharedRef<FUnrealEndpointManager>> MaybeUnrealEndpointManager = EndpointModule->GetEndpointManagerIfValid();
+		if (!MaybeUnrealEndpointManager.IsSet())
+		{
+			return;
+		}
 
-	// Wait for all the connections to actually close
-	for (const TSharedRef<FUnrealEndpoint>& Endpoint : Endpoints)
-	{
-		const bool bIsStopped = Endpoint->WaitForConnectionState(FUnrealEndpoint::EConnectionState::Disconnected, 3'000);
+		TSharedRef<FUnrealEndpointManager> UnrealEndpointManager = MaybeUnrealEndpointManager.GetValue();
+		const TArray<TSharedRef<FUnrealEndpoint>> Endpoints = UnrealEndpointManager->GetEndpoints();
 
-		if (!bIsStopped)
+		for (const TSharedRef<FUnrealEndpoint>& Endpoint : Endpoints)
 		{
-			FString EndpointInfoString = UnrealEndpointInfoToString(Endpoint->GetInfo());
-			UE_LOG(LogCaptureManagerIngestJobProcessor, Warning, TEXT("Failed to stop unreal endpoint: %s"), *EndpointInfoString);
+			Endpoint->StopConnection();
 		}
-	}
-}
 
+		// Wait for all the connections to actually close
+		for (const TSharedRef<FUnrealEndpoint>& Endpoint : Endpoints)
+		{
+			const bool bIsStopped = Endpoint->WaitForConnectionState(FUnrealEndpoint::EConnectionState::Disconnected, 3'000);
+
+			if (!bIsStopped)
+			{
+				FString EndpointInfoString = UnrealEndpointInfoToString(Endpoint->GetInfo());
+				UE_LOG(LogCaptureManagerIngestJobProcessor, Warning, TEXT("Failed to stop unreal endpoint: %s"), *EndpointInfoString);
+			}
+		}
+
+	}
 } // namespace Private
 
 TSharedRef<FIngestJobProcessor> FIngestJobProcessor::MakeInstance(int32 NumExecutors)
@@ -206,4 +218,4 @@ FIngestJobProcessor::FProcessingStateChanged& FIngestJobProcessor::OnProcessingS
 	return ProcessingStateChanged;
 }
 
-}
+}
\ No newline at end of file
diff --git a/Engine/Plugins/VirtualProduction/CaptureManager/CaptureManagerApp/Source/CaptureManagerUnrealEndpoint/Private/CaptureCommunicationManager.cpp b/Engine/Plugins/VirtualProduction/CaptureManager/CaptureManagerApp/Source/CaptureManagerUnrealEndpoint/Private/CaptureCommunicationManager.cpp
--- a/Engine/Plugins/VirtualProduction/CaptureManager/CaptureManagerApp/Source/CaptureManagerUnrealEndpoint/Private/CaptureCommunicationManager.cpp
+++ b/Engine/Plugins/VirtualProduction/CaptureManager/CaptureManagerApp/Source/CaptureManagerUnrealEndpoint/Private/CaptureCommunicationManager.cpp
@@ -129,7 +129,11 @@ TWeakPtr<FUploader> FCommunicationManager::GetUploader(const FGuid& InCaptureSou
 void FCommunicationManager::RemoveUploader(FGuid InCaptureSourceId)
 {
 	FScopeLock Lock(&Mutex);
-	Uploaders.Remove(InCaptureSourceId);
+	TSharedPtr<FUploader> Uploader;
+	if (Uploaders.RemoveAndCopyValue(InCaptureSourceId, Uploader))
+	{
+		Uploader->SetUploadHandlers(nullptr, nullptr);
+	}
 }
 
 bool FCommunicationManager::IsConnected() const
diff --git a/Engine/Plugins/VirtualProduction/CaptureManager/CaptureManagerApp/Source/CaptureManagerUnrealEndpoint/Private/CaptureManagerUnrealEndpointModule.cpp b/Engine/Plugins/VirtualProduction/CaptureManager/CaptureManagerApp/Source/CaptureManagerUnrealEndpoint/Private/CaptureManagerUnrealEndpointModule.cpp
--- a/Engine/Plugins/VirtualProduction/CaptureManager/CaptureManagerApp/Source/CaptureManagerUnrealEndpoint/Private/CaptureManagerUnrealEndpointModule.cpp
+++ b/Engine/Plugins/VirtualProduction/CaptureManager/CaptureManagerApp/Source/CaptureManagerUnrealEndpoint/Private/CaptureManagerUnrealEndpointModule.cpp
@@ -17,4 +17,14 @@ TSharedRef<UE::CaptureManager::FUnrealEndpointManager> FCaptureManagerUnrealEndp
 	return EndpointManager.ToSharedRef();
 }
 
+TOptional<TSharedRef<UE::CaptureManager::FUnrealEndpointManager>> FCaptureManagerUnrealEndpointModule::GetEndpointManagerIfValid() const
+{
+	if (EndpointManager.IsValid())
+	{
+		return EndpointManager.ToSharedRef();
+	}
+
+	return {};
+}
+
 IMPLEMENT_MODULE(FCaptureManagerUnrealEndpointModule, CaptureManagerUnrealEndpoint);
\ No newline at end of file
diff --git a/Engine/Plugins/VirtualProduction/CaptureManager/CaptureManagerApp/Source/CaptureManagerUnrealEndpoint/Public/CaptureManagerUnrealEndpointModule.h b/Engine/Plugins/VirtualProduction/CaptureManager/CaptureManagerApp/Source/CaptureManagerUnrealEndpoint/Public/CaptureManagerUnrealEndpointModule.h
--- a/Engine/Plugins/VirtualProduction/CaptureManager/CaptureManagerApp/Source/CaptureManagerUnrealEndpoint/Public/CaptureManagerUnrealEndpointModule.h
+++ b/Engine/Plugins/VirtualProduction/CaptureManager/CaptureManagerApp/Source/CaptureManagerUnrealEndpoint/Public/CaptureManagerUnrealEndpointModule.h
@@ -14,6 +14,7 @@ public:
 	virtual void ShutdownModule() override;
 
 	TSharedRef<UE::CaptureManager::FUnrealEndpointManager> GetEndpointManager() const;
+	TOptional<TSharedRef<UE::CaptureManager::FUnrealEndpointManager>> GetEndpointManagerIfValid() const;
 
 private:
 	TSharedPtr<UE::CaptureManager::FUnrealEndpointManager> EndpointManager;
diff --git a/Engine/Plugins/VirtualProduction/CaptureManager/CaptureManagerApp/Source/IngestLiveLinkDevice/Private/BaseIngestLiveLinkDevice.cpp b/Engine/Plugins/VirtualProduction/CaptureManager/CaptureManagerApp/Source/IngestLiveLinkDevice/Private/BaseIngestLiveLinkDevice.cpp
--- a/Engine/Plugins/VirtualProduction/CaptureManager/CaptureManagerApp/Source/IngestLiveLinkDevice/Private/BaseIngestLiveLinkDevice.cpp
+++ b/Engine/Plugins/VirtualProduction/CaptureManager/CaptureManagerApp/Source/IngestLiveLinkDevice/Private/BaseIngestLiveLinkDevice.cpp
@@ -30,6 +30,50 @@ static void TryCleanUpWorkingDirectory(const FString& InWorkingDirectory)
 
 } // namespace UE::CaptureManager::Private
 
+class FTaskWaiter {
+public:
+	FTaskWaiter() : TaskCounter(0)
+	{
+
+	}
+
+	~FTaskWaiter()
+	{
+		check(TaskCounter.Load() <= CanCreateTaskFlag);
+	}
+
+	bool CreateTask()
+	{
+		uint32 OldCounter = TaskCounter.IncrementExchange();
+		
+		if ((OldCounter & CanCreateTaskFlag) > 0)
+		{
+			TaskCounter.DecrementExchange();
+			return false;
+		}
+		return true;
+	}
+
+	void FinishTask()
+	{
+		TaskCounter.DecrementExchange();
+	}
+	
+	void WaitForAll()
+	{
+		TaskCounter.AddExchange(CanCreateTaskFlag);
+
+		while (TaskCounter > CanCreateTaskFlag)
+		{
+			FPlatformProcess::Sleep(0.1f);
+		}
+	}
+
+
+private: 
+	TAtomic<uint32> TaskCounter;
+	const uint32 CanCreateTaskFlag = 0x80000000;
+};
 
 class UBaseIngestLiveLinkDevice::FImpl : public TSharedFromThis<UBaseIngestLiveLinkDevice::FImpl>
 {
@@ -39,6 +83,7 @@ public:
 
 	void IngestTake(int32 InTakeId, const UIngestCapability_Options* InIngestOptions, TSharedPtr<UE::CaptureManager::FTaskProgress> InTaskProgress);
 	void CancelIngest(int32 InTakeId);
+	void OnDeviceRemoved(TArray<int32> TakeIdentifiers);
 
 private:
 
@@ -59,6 +104,12 @@ private:
 	TStrongObjectPtr<UBaseIngestLiveLinkDevice> Owner;
 	TMap<int32, TSharedRef<FCaptureDataConverter>> Converters;
 	TMap<int32, FTakeUploadInfo> TakeUploads;
+
+	// Every invocation of `IngestTake` function will increase a counter in TaskWaiter, registering that task is running.
+	// Once function finishes it's work, counter is decremented.
+	// And if it happens that device is being removed before counter is equal to zero, TaskWaiter will busy wait for it to reach zero.
+	// Ensures that 'IngestTake' does not access potentially destroyed device.
+	FTaskWaiter TaskWaiter;
 };
 
 UBaseIngestLiveLinkDevice::FImpl::FImpl(UBaseIngestLiveLinkDevice* InOwner) :
@@ -66,8 +117,41 @@ UBaseIngestLiveLinkDevice::FImpl::FImpl(UBaseIngestLiveLinkDevice* InOwner) :
 {
 }
 
+void UBaseIngestLiveLinkDevice::FImpl::OnDeviceRemoved(TArray<int32> TakeIdentifiers)
+{	
+	for (const int32 TakeIdentifier : TakeIdentifiers)
+	{
+		CancelIngest(TakeIdentifier);
+	}
+
+	while (true)
+	{
+		FScopeLock TakeUploadsRemoveLock(&Mutex);
+		int32 TakesInProgressCount = TakeUploads.Num();
+		TakeUploadsRemoveLock.Unlock();
+
+		if (TakesInProgressCount == 0)
+		{
+			break;
+		}
+		FPlatformProcess::Sleep(0.1f);
+	}
+
+	TaskWaiter.WaitForAll();
+}
+
 void UBaseIngestLiveLinkDevice::FImpl::IngestTake(int32 InTakeId, const UIngestCapability_Options* InIngestOptions, TSharedPtr<UE::CaptureManager::FTaskProgress> InTaskProgress)
 {
+	if (!TaskWaiter.CreateTask())
+	{
+		return;
+	}
+
+	ON_SCOPE_EXIT
+	{
+		TaskWaiter.FinishTask();
+	};
+
 	TStrongObjectPtr<const UIngestCapability_Options> IngestOptions(InIngestOptions);
 
 	using namespace UE::CaptureManager;
@@ -319,6 +403,8 @@ void UBaseIngestLiveLinkDevice::OnDeviceAdded()
 
 void UBaseIngestLiveLinkDevice::OnDeviceRemoved()
 {
+	Impl->OnDeviceRemoved(Execute_GetTakeIdentifiers(this));
+
 	RemoveAllTakes();
 
 	Super::OnDeviceRemoved();
