Passing UpdateTakeList callback as a parameter

* Removed the need to create the callback in the ingest capability before invoking UpdateTakeList
* UpdateTakeList accepts the callback as a parameter

diff --git a/Engine/Plugins/VirtualProduction/CaptureManager/CaptureManagerApp/Source/CaptureManagerEditor/Private/Scripting/ULiveLinkHubCaptureDevice.cpp b/Engine/Plugins/VirtualProduction/CaptureManager/CaptureManagerApp/Source/CaptureManagerEditor/Private/Scripting/ULiveLinkHubCaptureDevice.cpp
--- a/Engine/Plugins/VirtualProduction/CaptureManager/CaptureManagerApp/Source/CaptureManagerEditor/Private/Scripting/ULiveLinkHubCaptureDevice.cpp
+++ b/Engine/Plugins/VirtualProduction/CaptureManager/CaptureManagerApp/Source/CaptureManagerEditor/Private/Scripting/ULiveLinkHubCaptureDevice.cpp
@@ -273,13 +273,12 @@ TObjectPtr<UIngestCapability_ProcessResult> ULiveLinkHubCaptureDevice::FImpl::Fe
 		EDelegateExecutionThread::InternalThread
 	);
 
-	UIngestCapability_UpdateTakeListCallback* UpdateTakeListCallback =
-		ILiveLinkDeviceCapability_Ingest::Execute_CreateUpdateTakeListCallback(BackgroundDevice);
+	UIngestCapability_UpdateTakeListCallback* UpdateTakeListCallback = NewObject<UIngestCapability_UpdateTakeListCallback>();
 	check(UpdateTakeListCallback);
 
 	UpdateTakeListCallback->Callback = MoveTemp(OnTakesListUpdated);
 
-	ILiveLinkDeviceCapability_Ingest::Execute_UpdateTakeList(BackgroundDevice);
+	ILiveLinkDeviceCapability_Ingest::Execute_UpdateTakeList(BackgroundDevice, UpdateTakeListCallback);
 	TakeListUpdatedEvent->Wait();
 
 	OutTakes = MoveTemp(Takes);
diff --git a/Engine/Plugins/VirtualProduction/CaptureManager/CaptureManagerApp/Source/CaptureManagerEditor/Private/UI/TakesView.cpp b/Engine/Plugins/VirtualProduction/CaptureManager/CaptureManagerApp/Source/CaptureManagerEditor/Private/UI/TakesView.cpp
--- a/Engine/Plugins/VirtualProduction/CaptureManager/CaptureManagerApp/Source/CaptureManagerEditor/Private/UI/TakesView.cpp
+++ b/Engine/Plugins/VirtualProduction/CaptureManager/CaptureManagerApp/Source/CaptureManagerEditor/Private/UI/TakesView.cpp
@@ -502,8 +502,9 @@ void FTakesView::CaptureDeviceStarted(FGuid InCaptureDeviceId)
 						}
 					}));
 
-			ILiveLinkDeviceCapability_Ingest::Execute_CreateUpdateTakeListCallback(Device)->Callback = MoveTemp(Callback);
-			ILiveLinkDeviceCapability_Ingest::Execute_UpdateTakeList(Device);
+			UIngestCapability_UpdateTakeListCallback* UpdateTakeListCallback = NewObject<UIngestCapability_UpdateTakeListCallback>();
+			UpdateTakeListCallback->Callback = MoveTemp(Callback);
+			ILiveLinkDeviceCapability_Ingest::Execute_UpdateTakeList(Device, UpdateTakeListCallback);
 		});
 }
 
diff --git a/Engine/Plugins/VirtualProduction/CaptureManager/CaptureManagerApp/Source/ExampleLiveLinkDevices/Private/NetworkIngest/ExampleNetworkIngestDevice.cpp b/Engine/Plugins/VirtualProduction/CaptureManager/CaptureManagerApp/Source/ExampleLiveLinkDevices/Private/NetworkIngest/ExampleNetworkIngestDevice.cpp
--- a/Engine/Plugins/VirtualProduction/CaptureManager/CaptureManagerApp/Source/ExampleLiveLinkDevices/Private/NetworkIngest/ExampleNetworkIngestDevice.cpp
+++ b/Engine/Plugins/VirtualProduction/CaptureManager/CaptureManagerApp/Source/ExampleLiveLinkDevices/Private/NetworkIngest/ExampleNetworkIngestDevice.cpp
@@ -62,7 +62,7 @@ FString UExampleNetworkIngestDevice::GetFullTakePath(UE::CaptureManager::FTakeId
 	return FString();
 }
 
-void UExampleNetworkIngestDevice::UpdateTakeList_Implementation()
+void UExampleNetworkIngestDevice::UpdateTakeList_Implementation(UIngestCapability_UpdateTakeListCallback* InCallback)
 {
 	// Query your device for a list of takes
 	// Populate an FTakeMetadata object per take
diff --git a/Engine/Plugins/VirtualProduction/CaptureManager/CaptureManagerApp/Source/ExampleLiveLinkDevices/Private/NetworkIngest/ExampleNetworkIngestDevice.h b/Engine/Plugins/VirtualProduction/CaptureManager/CaptureManagerApp/Source/ExampleLiveLinkDevices/Private/NetworkIngest/ExampleNetworkIngestDevice.h
--- a/Engine/Plugins/VirtualProduction/CaptureManager/CaptureManagerApp/Source/ExampleLiveLinkDevices/Private/NetworkIngest/ExampleNetworkIngestDevice.h
+++ b/Engine/Plugins/VirtualProduction/CaptureManager/CaptureManagerApp/Source/ExampleLiveLinkDevices/Private/NetworkIngest/ExampleNetworkIngestDevice.h
@@ -56,7 +56,7 @@ private:
 	virtual FString GetFullTakePath(UE::CaptureManager::FTakeId InTakeId) const override;
 
 	//~ Begin ULiveLinkDeviceCapability_Ingest interface
-	virtual void UpdateTakeList_Implementation() override;
+	virtual void UpdateTakeList_Implementation(UIngestCapability_UpdateTakeListCallback* InCallback) override;
 	virtual void RunIngestTake(int32 InTakeId,
 									const UIngestCapability_Options* InIngestOptions) override;
 	virtual void CancelIngest(int32 InTakeId) override;
diff --git a/Engine/Plugins/VirtualProduction/CaptureManager/CaptureManagerApp/Source/LiveLinkCapabilities/Private/Ingest/LiveLinkDeviceCapability_Ingest.cpp b/Engine/Plugins/VirtualProduction/CaptureManager/CaptureManagerApp/Source/LiveLinkCapabilities/Private/Ingest/LiveLinkDeviceCapability_Ingest.cpp
--- a/Engine/Plugins/VirtualProduction/CaptureManager/CaptureManagerApp/Source/LiveLinkCapabilities/Private/Ingest/LiveLinkDeviceCapability_Ingest.cpp
+++ b/Engine/Plugins/VirtualProduction/CaptureManager/CaptureManagerApp/Source/LiveLinkCapabilities/Private/Ingest/LiveLinkDeviceCapability_Ingest.cpp
@@ -92,13 +92,6 @@ TArray<int32> ILiveLinkDeviceCapability_Ingest::GetTakeIdentifiers_Implementatio
 	return TakeIdentifiers;
 }
 
-UIngestCapability_UpdateTakeListCallback* ILiveLinkDeviceCapability_Ingest::CreateUpdateTakeListCallback_Implementation()
-{
-	UpdateTakeListCallback = TStrongObjectPtr(NewObject<UIngestCapability_UpdateTakeListCallback>());
-
-	return UpdateTakeListCallback.Get();
-}
-
 void ILiveLinkDeviceCapability_Ingest::ExecuteProcessFinishedReporter(int32 InTakeIdentifier,
 																	  TValueOrError<void, FIngestCapability_Error> InMaybeError)
 {
@@ -145,15 +138,13 @@ void ILiveLinkDeviceCapability_Ingest::ExecuteProcessProgressReporter(int32 InTa
 	IngestProcessHandle->ProcessProgressReporter(InTakeIdentifier, InProgress);
 }
 
-void ILiveLinkDeviceCapability_Ingest::ExecuteUpdateTakeListCallback(const TArray<int32>& InTakeIdentifiers)
+void ILiveLinkDeviceCapability_Ingest::ExecuteUpdateTakeListCallback(UIngestCapability_UpdateTakeListCallback* InCallback, const TArray<int32>& InTakeIdentifiers)
 {
-	if (ensure(UpdateTakeListCallback))
+	if (ensure(InCallback))
 	{
-		UpdateTakeListCallback->DynamicCallback.ExecuteIfBound(InTakeIdentifiers);
-		UpdateTakeListCallback->Callback(InTakeIdentifiers);
+		InCallback->DynamicCallback.ExecuteIfBound(InTakeIdentifiers);
+		InCallback->Callback(InTakeIdentifiers);
 	}
-
-	UpdateTakeListCallback = nullptr;
 }
 
 int32 ILiveLinkDeviceCapability_Ingest::AddTake(FTakeMetadata InTakeMetadata)
diff --git a/Engine/Plugins/VirtualProduction/CaptureManager/CaptureManagerApp/Source/LiveLinkCapabilities/Public/Ingest/LiveLinkDeviceCapability_Ingest.h b/Engine/Plugins/VirtualProduction/CaptureManager/CaptureManagerApp/Source/LiveLinkCapabilities/Public/Ingest/LiveLinkDeviceCapability_Ingest.h
--- a/Engine/Plugins/VirtualProduction/CaptureManager/CaptureManagerApp/Source/LiveLinkCapabilities/Public/Ingest/LiveLinkDeviceCapability_Ingest.h
+++ b/Engine/Plugins/VirtualProduction/CaptureManager/CaptureManagerApp/Source/LiveLinkCapabilities/Public/Ingest/LiveLinkDeviceCapability_Ingest.h
@@ -49,7 +49,7 @@ public:
 	UIngestCapability_ProcessHandle* PrepareIngestProcess(int32 InTakeId);
 
 	UFUNCTION(BlueprintCallable, BlueprintNativeEvent, Category = "Live Link Device|Ingest")
-	void UpdateTakeList();
+	void UpdateTakeList(UIngestCapability_UpdateTakeListCallback* InCallback);
 
 	UFUNCTION(BlueprintCallable, BlueprintNativeEvent, Category = "Live Link Device|Ingest")
 	UIngestCapability_TakeInformation* GetTakeInformation(int32 InTakeId) const;
@@ -57,14 +57,10 @@ public:
 	UFUNCTION(BlueprintCallable, BlueprintNativeEvent, Category = "Live Link Device|Ingest")
 	TArray<int32> GetTakeIdentifiers() const;
 
-	UFUNCTION(BlueprintCallable, BlueprintNativeEvent, Category = "Live Link Device|Ingest")
-	UIngestCapability_UpdateTakeListCallback* CreateUpdateTakeListCallback();
-
 	UIngestCapability_ProcessHandle* PrepareIngestProcess_Implementation(int32 InTakeId);
-	virtual void UpdateTakeList_Implementation() PURE_VIRTUAL(ILiveLinkDeviceCapability_Ingest::UpdateTakeList_Implementation);
+	virtual void UpdateTakeList_Implementation(UIngestCapability_UpdateTakeListCallback* InCallback) PURE_VIRTUAL(ILiveLinkDeviceCapability_Ingest::UpdateTakeList_Implementation);
 	UIngestCapability_TakeInformation* GetTakeInformation_Implementation(int32 InTakeId) const;
 	TArray<int32> GetTakeIdentifiers_Implementation() const;
-	UIngestCapability_UpdateTakeListCallback* CreateUpdateTakeListCallback_Implementation();
 
 	TOptional<FTakeMetadata> GetTakeMetadata(int32 InTakeId) const;
 
@@ -72,7 +68,7 @@ protected:
 
 	void ExecuteProcessFinishedReporter(int32 InTakeIdentifier, TValueOrError<void, FIngestCapability_Error> InMaybeError);
 	void ExecuteProcessProgressReporter(int32 InTakeIdentifier, double InProgress);
-	void ExecuteUpdateTakeListCallback(const TArray<int32>& InTakeIdentifiers);
+	void ExecuteUpdateTakeListCallback(UIngestCapability_UpdateTakeListCallback* InCallback, const TArray<int32>& InTakeIdentifiers);
 	
 	int32 AddTake(FTakeMetadata InTakeMetadata);
 	void RemoveTake(int32 InTakeId);
@@ -93,8 +89,6 @@ private:
 	mutable FCriticalSection IngestProcessMutex;
 	TMap<int32, TStrongObjectPtr<UIngestCapability_ProcessHandle>> IngestProcessHandles;
 
-	TStrongObjectPtr<UIngestCapability_UpdateTakeListCallback> UpdateTakeListCallback;
-
 	std::atomic_int32_t CurrentTakeId = 0;
 
 	friend class UIngestCapability_ProcessHandle;
diff --git a/Engine/Plugins/VirtualProduction/CaptureManager/CaptureManagerDevices/Source/CPSLiveLinkDevice/Private/LiveLinkFaceDevice.cpp b/Engine/Plugins/VirtualProduction/CaptureManager/CaptureManagerDevices/Source/CPSLiveLinkDevice/Private/LiveLinkFaceDevice.cpp
--- a/Engine/Plugins/VirtualProduction/CaptureManager/CaptureManagerDevices/Source/CPSLiveLinkDevice/Private/LiveLinkFaceDevice.cpp
+++ b/Engine/Plugins/VirtualProduction/CaptureManager/CaptureManagerDevices/Source/CPSLiveLinkDevice/Private/LiveLinkFaceDevice.cpp
@@ -67,7 +67,7 @@ FString ULiveLinkFaceDevice::GetFullTakePath(UE::CaptureManager::FTakeId InTakeI
 	return FString();
 }
 
-void ULiveLinkFaceDevice::UpdateTakeList_Implementation()
+void ULiveLinkFaceDevice::UpdateTakeList_Implementation(UIngestCapability_UpdateTakeListCallback* InCallback)
 {
 	using namespace UE::CaptureManager;
 
@@ -91,7 +91,7 @@ void ULiveLinkFaceDevice::UpdateTakeList_Implementation()
 
 	FetchPreIngestFiles(MoveTemp(NameToIdMap));
 
-	ExecuteUpdateTakeListCallback(Execute_GetTakeIdentifiers(this));
+	ExecuteUpdateTakeListCallback(InCallback, Execute_GetTakeIdentifiers(this));
 }
 
 void ULiveLinkFaceDevice::RunIngestTake(int32 InTakeId, const UIngestCapability_Options* InIngestOptions)
diff --git a/Engine/Plugins/VirtualProduction/CaptureManager/CaptureManagerDevices/Source/CPSLiveLinkDevice/Private/LiveLinkFaceDevice.h b/Engine/Plugins/VirtualProduction/CaptureManager/CaptureManagerDevices/Source/CPSLiveLinkDevice/Private/LiveLinkFaceDevice.h
--- a/Engine/Plugins/VirtualProduction/CaptureManager/CaptureManagerDevices/Source/CPSLiveLinkDevice/Private/LiveLinkFaceDevice.h
+++ b/Engine/Plugins/VirtualProduction/CaptureManager/CaptureManagerDevices/Source/CPSLiveLinkDevice/Private/LiveLinkFaceDevice.h
@@ -65,7 +65,7 @@ private:
 	virtual FString GetFullTakePath(UE::CaptureManager::FTakeId InTakeId) const override;
 
 	//~ Begin ULiveLinkDeviceCapability_Ingest interface
-	virtual void UpdateTakeList_Implementation() override;
+	virtual void UpdateTakeList_Implementation(UIngestCapability_UpdateTakeListCallback* InCallback) override;
 	virtual void RunIngestTake(int32 InTakeId, const UIngestCapability_Options* InIngestOptions) override;
 	virtual void CancelIngest(int32 InTakeId) override;
 	//~ End ULiveLinkDeviceCapability_Ingest interface
diff --git a/Engine/Plugins/VirtualProduction/CaptureManager/CaptureManagerDevices/Source/FilesystemLiveLinkDevice/Private/FilesystemLiveLinkDevice.cpp b/Engine/Plugins/VirtualProduction/CaptureManager/CaptureManagerDevices/Source/FilesystemLiveLinkDevice/Private/FilesystemLiveLinkDevice.cpp
--- a/Engine/Plugins/VirtualProduction/CaptureManager/CaptureManagerDevices/Source/FilesystemLiveLinkDevice/Private/FilesystemLiveLinkDevice.cpp
+++ b/Engine/Plugins/VirtualProduction/CaptureManager/CaptureManagerDevices/Source/FilesystemLiveLinkDevice/Private/FilesystemLiveLinkDevice.cpp
@@ -88,9 +88,9 @@ void UFilesystemLiveLinkDevice::RunIngestTake(int32 InTakeId, const UIngestCapab
 	});
 }
 
-void UFilesystemLiveLinkDevice::UpdateTakeList_Implementation()
+void UFilesystemLiveLinkDevice::UpdateTakeList_Implementation(UIngestCapability_UpdateTakeListCallback* InCallback)
 {
-	AsyncTask(ENamedThreads::AnyThread, [this]()
+	AsyncTask(ENamedThreads::AnyThread, [this, Callback = TStrongObjectPtr<UIngestCapability_UpdateTakeListCallback>(InCallback)]()
 	{
 		RemoveAllTakes();
 
@@ -130,7 +130,7 @@ void UFilesystemLiveLinkDevice::UpdateTakeList_Implementation()
 			}
 		}
 
-		ExecuteUpdateTakeListCallback(Execute_GetTakeIdentifiers(this));
+		ExecuteUpdateTakeListCallback(Callback.Get(), Execute_GetTakeIdentifiers(this));
 	});
 }
 
diff --git a/Engine/Plugins/VirtualProduction/CaptureManager/CaptureManagerDevices/Source/FilesystemLiveLinkDevice/Private/FilesystemLiveLinkDevice.h b/Engine/Plugins/VirtualProduction/CaptureManager/CaptureManagerDevices/Source/FilesystemLiveLinkDevice/Private/FilesystemLiveLinkDevice.h
--- a/Engine/Plugins/VirtualProduction/CaptureManager/CaptureManagerDevices/Source/FilesystemLiveLinkDevice/Private/FilesystemLiveLinkDevice.h
+++ b/Engine/Plugins/VirtualProduction/CaptureManager/CaptureManagerDevices/Source/FilesystemLiveLinkDevice/Private/FilesystemLiveLinkDevice.h
@@ -55,7 +55,7 @@ private:
 
 	//~ Begin ULiveLinkDeviceCapability_Ingest interface
 	virtual void RunIngestTake(int32 InTakeId, const UIngestCapability_Options* InIngestOptions) override;
-	virtual void UpdateTakeList_Implementation() override;
+	virtual void UpdateTakeList_Implementation(UIngestCapability_UpdateTakeListCallback* InCallback) override;
 	//~ End ULiveLinkDeviceCapability_Ingest interface
 
 	//~ Begin ILiveLinkDeviceCapability_Connection interface
diff --git a/Engine/Plugins/VirtualProduction/CaptureManager/CaptureManagerDevices/Source/MonoVideoLiveLinkDevice/Private/MonoVideoLiveLinkDevice.cpp b/Engine/Plugins/VirtualProduction/CaptureManager/CaptureManagerDevices/Source/MonoVideoLiveLinkDevice/Private/MonoVideoLiveLinkDevice.cpp
--- a/Engine/Plugins/VirtualProduction/CaptureManager/CaptureManagerDevices/Source/MonoVideoLiveLinkDevice/Private/MonoVideoLiveLinkDevice.cpp
+++ b/Engine/Plugins/VirtualProduction/CaptureManager/CaptureManagerDevices/Source/MonoVideoLiveLinkDevice/Private/MonoVideoLiveLinkDevice.cpp
@@ -116,7 +116,7 @@ void UMonoVideoLiveLinkDevice::RunIngestTake(int32 InTakeId, const UIngestCapabi
 	});
 }
 
-void UMonoVideoLiveLinkDevice::UpdateTakeList_Implementation()
+void UMonoVideoLiveLinkDevice::UpdateTakeList_Implementation(UIngestCapability_UpdateTakeListCallback* InCallback)
 {
 	RemoveAllTakes();
 
@@ -159,7 +159,7 @@ void UMonoVideoLiveLinkDevice::UpdateTakeList_Implementation()
 		}
 	}
 
-	ExecuteUpdateTakeListCallback(Execute_GetTakeIdentifiers(this));
+	ExecuteUpdateTakeListCallback(InCallback, Execute_GetTakeIdentifiers(this));
 }
 
 ELiveLinkDeviceConnectionStatus UMonoVideoLiveLinkDevice::GetConnectionStatus_Implementation() const
diff --git a/Engine/Plugins/VirtualProduction/CaptureManager/CaptureManagerDevices/Source/MonoVideoLiveLinkDevice/Private/MonoVideoLiveLinkDevice.h b/Engine/Plugins/VirtualProduction/CaptureManager/CaptureManagerDevices/Source/MonoVideoLiveLinkDevice/Private/MonoVideoLiveLinkDevice.h
--- a/Engine/Plugins/VirtualProduction/CaptureManager/CaptureManagerDevices/Source/MonoVideoLiveLinkDevice/Private/MonoVideoLiveLinkDevice.h
+++ b/Engine/Plugins/VirtualProduction/CaptureManager/CaptureManagerDevices/Source/MonoVideoLiveLinkDevice/Private/MonoVideoLiveLinkDevice.h
@@ -86,7 +86,7 @@ private:
 
 	//~ Begin ULiveLinkDeviceCapability_Ingest interface
 	virtual void RunIngestTake(int32 InTakeId, const UIngestCapability_Options* InIngestOptions) override;
-	virtual void UpdateTakeList_Implementation() override;
+	virtual void UpdateTakeList_Implementation(UIngestCapability_UpdateTakeListCallback* InCallback) override;
 	//~ End ULiveLinkDeviceCapability_Ingest interface
 
 	//~ Begin ILiveLinkDeviceCapability_Connection interface
diff --git a/Engine/Plugins/VirtualProduction/CaptureManager/CaptureManagerDevices/Source/StereoVideoLiveLinkDevice/Private/StereoVideoLiveLinkDevice.cpp b/Engine/Plugins/VirtualProduction/CaptureManager/CaptureManagerDevices/Source/StereoVideoLiveLinkDevice/Private/StereoVideoLiveLinkDevice.cpp
--- a/Engine/Plugins/VirtualProduction/CaptureManager/CaptureManagerDevices/Source/StereoVideoLiveLinkDevice/Private/StereoVideoLiveLinkDevice.cpp
+++ b/Engine/Plugins/VirtualProduction/CaptureManager/CaptureManagerDevices/Source/StereoVideoLiveLinkDevice/Private/StereoVideoLiveLinkDevice.cpp
@@ -142,7 +142,7 @@ int32 UStereoVideoLiveLinkDevice::FTakeWithComponents::CountComponents(UStereoVi
 	return Count;
 }
 
-void UStereoVideoLiveLinkDevice::UpdateTakeList_Implementation()
+void UStereoVideoLiveLinkDevice::UpdateTakeList_Implementation(UIngestCapability_UpdateTakeListCallback* InCallback)
 {
 	RemoveAllTakes();
 
@@ -178,7 +178,7 @@ void UStereoVideoLiveLinkDevice::UpdateTakeList_Implementation()
 		}
 	}
 
-	ExecuteUpdateTakeListCallback(Execute_GetTakeIdentifiers(this));
+	ExecuteUpdateTakeListCallback(InCallback, Execute_GetTakeIdentifiers(this));
 }
 
 TMap<FString, UStereoVideoLiveLinkDevice::FTakeWithComponents> UStereoVideoLiveLinkDevice::DiscoverTakes(FString StoragePath)
diff --git a/Engine/Plugins/VirtualProduction/CaptureManager/CaptureManagerDevices/Source/StereoVideoLiveLinkDevice/Private/StereoVideoLiveLinkDevice.h b/Engine/Plugins/VirtualProduction/CaptureManager/CaptureManagerDevices/Source/StereoVideoLiveLinkDevice/Private/StereoVideoLiveLinkDevice.h
--- a/Engine/Plugins/VirtualProduction/CaptureManager/CaptureManagerDevices/Source/StereoVideoLiveLinkDevice/Private/StereoVideoLiveLinkDevice.h
+++ b/Engine/Plugins/VirtualProduction/CaptureManager/CaptureManagerDevices/Source/StereoVideoLiveLinkDevice/Private/StereoVideoLiveLinkDevice.h
@@ -142,7 +142,7 @@ private:
 
 	//~ Begin ULiveLinkDeviceCapability_Ingest interface
 	virtual void RunIngestTake(int32 InTakeId, const UIngestCapability_Options* InIngestOptions) override;
-	virtual void UpdateTakeList_Implementation() override;
+	virtual void UpdateTakeList_Implementation(UIngestCapability_UpdateTakeListCallback* InCallback) override;
 	//~ End ULiveLinkDeviceCapability_Ingest interface
 
 	//~ Begin ILiveLinkDeviceCapability_Connection interface
