[CaptureManager] Added ability to refresh takes for the selected device

diff --git a/Engine/Plugins/VirtualProduction/CaptureManager/CaptureManagerApp/Source/CaptureManagerEditor/Private/UI/TakesView.cpp b/Engine/Plugins/VirtualProduction/CaptureManager/CaptureManagerApp/Source/CaptureManagerEditor/Private/UI/TakesView.cpp
--- a/Engine/Plugins/VirtualProduction/CaptureManager/CaptureManagerApp/Source/CaptureManagerEditor/Private/UI/TakesView.cpp
+++ b/Engine/Plugins/VirtualProduction/CaptureManager/CaptureManagerApp/Source/CaptureManagerEditor/Private/UI/TakesView.cpp
@@ -43,6 +43,7 @@ void STakesView::Construct(
 	IngestPipelineManager = InIngestPipelineManager.ToSharedPtr();
 	IngestJobSettingsManager = InIngestJobsSettingsManager.ToSharedPtr();
 	OnAddTakesToIngestQueue = InArgs._OnAddTakesToIngestQueue;
+	OnRefreshTakes = InArgs._OnRefreshTakes;
 	CurrentPipeline = IngestPipelineManager->GetSelectedPipeline();
 
 	const TStrongObjectPtr<const UIngestJobSettingsBase> SelectedPreset(IngestJobSettingsManager->GetSelectedPreset());
@@ -58,6 +59,11 @@ void STakesView::Construct(
 		PipelineNames.Emplace(MakeShared<FText>(Pipeline.DisplayName));
 	}
 
+	RefreshButton = SNew(SPositiveActionButton)
+		.Icon(FAppStyle::Get().GetBrush("Icons.Refresh"))
+		.ToolTipText(LOCTEXT("RefreshTakes_Tooltip", "Refresh takes for the selected device"))
+		.OnClicked_Raw(this, &STakesView::OnRefreshButtonClicked);
+
 	AddToQueueButton = SNew(SPositiveActionButton)
 		.Icon(FAppStyle::Get().GetBrush("Icons.Plus"))
 		.Text(LOCTEXT("AddToQueueText", "Add to Queue"))
@@ -77,6 +83,13 @@ void STakesView::Construct(
 				.Padding(2.0f)
 				[
 					SNew(SHorizontalBox)
+						+ SHorizontalBox::Slot()
+						.Padding(2.0f)
+						.AutoWidth()
+						.VAlign(EVerticalAlignment::VAlign_Center)
+						[
+							RefreshButton.ToSharedRef()
+						]
 						+ SHorizontalBox::Slot()
 						.Padding(2.0f)
 						.AutoWidth()
@@ -164,6 +177,13 @@ TArray<TStrongObjectPtr<UTakeVirtualAsset>> STakesView::GetSelectedTakeAssets()
 	return TakeAssets;
 }
 
+FReply STakesView::OnRefreshButtonClicked()
+{
+	OnRefreshTakes.ExecuteIfBound();
+
+	return FReply::Handled();
+}
+
 FReply STakesView::OnAddToQueueButtonClicked()
 {
 	check(IngestJobSettingsManager);
@@ -330,6 +350,34 @@ void FTakesView::AddTakesToIngestQueue(const TArray<TStrongObjectPtr<UTakeVirtua
 	}
 }
 
+void FTakesView::RefreshTakes()
+{
+	using namespace UE::CaptureManager;
+
+	if (SelectedItem.IsValid())
+	{
+		if (SelectedItem->Implements<ULiveLinkDeviceCapability_Ingest>())
+		{
+			FIngestUpdateTakeListCallback Callback = FIngestUpdateTakeListCallback::Type::CreateLambda(
+				[this, InCaptureDeviceId = SelectedItem->GetDeviceId()](const TArray<FTakeId>& InTakeIds)
+				{
+					CaptureDeviceTakesMap.Remove(InCaptureDeviceId);
+
+					TArray<TStrongObjectPtr<UTakeVirtualAsset>> TakeObjects = MakeTakeObjects(InCaptureDeviceId, InTakeIds);
+
+					CaptureDeviceTakesMap.Add(InCaptureDeviceId, MoveTemp(TakeObjects));
+
+					RefreshTakesViewDelegate.ExecuteIfBound(true);
+				}
+			);
+
+			UIngestCapability_UpdateTakeListCallback* UpdateTakeListCallback = NewObject<UIngestCapability_UpdateTakeListCallback>();
+			UpdateTakeListCallback->Callback = MoveTemp(Callback);
+			ILiveLinkDeviceCapability_Ingest::Execute_UpdateTakeList(SelectedItem.Get(), UpdateTakeListCallback);
+		}
+	}
+}
+
 void FTakesView::CreateTakesView()
 {
 	TSharedPtr<FCaptureManagerPanelController> Controller = PanelController.Pin();
@@ -387,6 +435,7 @@ void FTakesView::CreateTakesView()
 
 	SAssignNew(TakesTileView, STakesView, MoveTemp(IngestPipelineManager), MoveTemp(IngestJobSettingsManager), MoveTemp(GetCurrentSelectionDelegate))
 		.OnAddTakesToIngestQueue_Raw(this, &FTakesView::AddTakesToIngestQueue)
+		.OnRefreshTakes_Raw(this, &FTakesView::RefreshTakes)
 		.TakesPickerConfig(TakesPickerConfig);
 }
 
diff --git a/Engine/Plugins/VirtualProduction/CaptureManager/CaptureManagerApp/Source/CaptureManagerEditor/Private/UI/TakesView.h b/Engine/Plugins/VirtualProduction/CaptureManager/CaptureManagerApp/Source/CaptureManagerEditor/Private/UI/TakesView.h
--- a/Engine/Plugins/VirtualProduction/CaptureManager/CaptureManagerApp/Source/CaptureManagerEditor/Private/UI/TakesView.h
+++ b/Engine/Plugins/VirtualProduction/CaptureManager/CaptureManagerApp/Source/CaptureManagerEditor/Private/UI/TakesView.h
@@ -34,10 +34,12 @@ class STakesView : public SCompoundWidget
 {
 public:
 	DECLARE_DELEGATE_TwoParams(FOnAddTakesToIngestQueue, const TArray<TStrongObjectPtr<UTakeVirtualAsset>>&, const UIngestJobSettingsBase&);
+	DECLARE_DELEGATE(FOnRefreshTakes);
 
 	SLATE_BEGIN_ARGS(STakesView) {}
 		SLATE_ARGUMENT(FAssetPickerConfig, TakesPickerConfig)
 		SLATE_EVENT(FOnAddTakesToIngestQueue, OnAddTakesToIngestQueue)
+		SLATE_EVENT(FOnRefreshTakes, OnRefreshTakes)
 	SLATE_END_ARGS()
 
 	STakesView();
@@ -52,12 +54,15 @@ public:
 	void SetAddToQueueButtonEnabled(bool bIsEnabled);
 
 private:
+	FReply OnRefreshButtonClicked();
 	FReply OnAddToQueueButtonClicked();
 	TSharedRef<SWidget> OnGeneratePipelineNameWidget(TSharedRef<FText> InText);
 	void OnPipelineSelectionChanged(TSharedPtr<FText> InText, ESelectInfo::Type InSelectType);
 	TArray<TStrongObjectPtr<UTakeVirtualAsset>> GetSelectedTakeAssets() const;
 
 	TUniquePtr<FGetCurrentSelectionDelegate> GetCurrentSelectionDelegate;
+	TSharedPtr<SPositiveActionButton> RefreshButton;
+	FOnRefreshTakes OnRefreshTakes;
 	TSharedPtr<SPositiveActionButton> AddToQueueButton;
 	FOnAddTakesToIngestQueue OnAddTakesToIngestQueue;
 
@@ -93,6 +98,7 @@ private:
 
 	EImageRotation ConvertOrientation(TOptional<FTakeMetadata::FVideo::EOrientation> InOrientation) const;
 	void AddTakesToIngestQueue(const TArray<TStrongObjectPtr<UTakeVirtualAsset>>& InTakeAssets, const UIngestJobSettingsBase& InJobPreset);
+	void RefreshTakes();
 
 	TArray<TStrongObjectPtr<UTakeVirtualAsset>> MakeTakeObjects(FGuid InCaptureDeviceId,
 																TOptional<TArray<UE::CaptureManager::FTakeId>> InTakeIds = TOptional<TArray<UE::CaptureManager::FTakeId>>());
