Invoking take added after reading and parsing each take

diff --git a/Engine/Plugins/VirtualProduction/CaptureManager/CaptureManagerApp/Source/CaptureManagerEditor/Private/UI/TakesView.cpp b/Engine/Plugins/VirtualProduction/CaptureManager/CaptureManagerApp/Source/CaptureManagerEditor/Private/UI/TakesView.cpp
--- a/Engine/Plugins/VirtualProduction/CaptureManager/CaptureManagerApp/Source/CaptureManagerEditor/Private/UI/TakesView.cpp
+++ b/Engine/Plugins/VirtualProduction/CaptureManager/CaptureManagerApp/Source/CaptureManagerEditor/Private/UI/TakesView.cpp
@@ -15,6 +15,7 @@
 #include "Widgets/Input/SButton.h"
 #include "Widgets/Input/SComboBox.h"
 #include "SPositiveActionButton.h"
+#include "Widgets/Images/SThrobber.h"
 
 #include "Ingest/LiveLinkDeviceCapability_Ingest.h"
 #include "Ingest/IngestCapability_Events.h"
@@ -24,6 +25,38 @@
 
 DEFINE_LOG_CATEGORY_STATIC(LogTakesView, Log, All);
 
+namespace UE::CaptureManager
+{
+template<typename ElementType, typename OtherElementType, typename Predicate>
+bool EqualTo(const TArray<ElementType>& InFirstArray, const TArray<OtherElementType>& InSecondArray, Predicate InPredicate)
+{
+	if (InFirstArray.Num() != InSecondArray.Num())
+	{
+		return false;
+	}
+
+	for (int64 Index = 0; Index < InFirstArray.Num(); ++Index)
+	{
+		bool bMatchFound = false;
+		for (int64 OtherIndex = 0; OtherIndex < InSecondArray.Num(); ++OtherIndex)
+		{
+			if (InPredicate(InFirstArray[Index], InSecondArray[OtherIndex]))
+			{
+				bMatchFound = true;
+				break;
+			}
+		}
+
+		if (!bMatchFound)
+		{
+			return false;
+		}
+	}
+
+	return true;
+}
+}
+
 static const FText JobPipelineText = LOCTEXT("JobPipelineText", "Pipeline");
 
 static const FString OuterPackageName = TEXT("CaptureDevices");
@@ -62,15 +95,16 @@ void STakesView::Construct(
 	RefreshButton = SNew(SPositiveActionButton)
 		.Icon(FAppStyle::Get().GetBrush("Icons.Refresh"))
 		.ToolTipText(LOCTEXT("RefreshTakes_Tooltip", "Refresh takes for the selected device"))
-		.OnClicked_Raw(this, &STakesView::OnRefreshButtonClicked);
+		.OnClicked(this, &STakesView::OnRefreshButtonClicked);
 
 	AddToQueueButton = SNew(SPositiveActionButton)
 		.Icon(FAppStyle::Get().GetBrush("Icons.Plus"))
 		.Text(LOCTEXT("AddToQueueText", "Add to Queue"))
 		.ToolTipText(LOCTEXT("AddToQueue_Tooltip", "Add selected takes to the ingest jobs list"))
-		.OnClicked_Raw(this, &STakesView::OnAddToQueueButtonClicked);
+		.OnClicked(this, &STakesView::OnAddToQueueButtonClicked);
 
 	IContentBrowserSingleton& ContentBrowser = FModuleManager::LoadModuleChecked<FContentBrowserModule>("ContentBrowser").Get();
+	static const FSlateBrush* CachedImage = FAppStyle::Get().GetBrush("UnrealCircle.Thick");
 
 	FAssetPickerConfig TakesPickerConfig = InArgs._TakesPickerConfig;
 
@@ -83,59 +117,90 @@ void STakesView::Construct(
 				.Padding(2.0f)
 				[
 					SNew(SHorizontalBox)
-						+ SHorizontalBox::Slot()
-						.Padding(2.0f)
-						.AutoWidth()
-						.VAlign(EVerticalAlignment::VAlign_Center)
+					+ SHorizontalBox::Slot()
+					.Padding(2.0f)
+					.AutoWidth()
+					.VAlign(EVerticalAlignment::VAlign_Center)
+					[
+						RefreshButton.ToSharedRef()
+					]
+					+ SHorizontalBox::Slot()
+					.Padding(2.0f)
+					.AutoWidth()
+					.VAlign(EVerticalAlignment::VAlign_Center)
+					[
+						SNew(STextBlock)
+							.Text(JobPipelineText)
+					]
+					+ SHorizontalBox::Slot()
+					.Padding(2.0f)
+					.HAlign(EHorizontalAlignment::HAlign_Fill)
+					.VAlign(EVerticalAlignment::VAlign_Center)
+					[
+						SNew(SComboBox<TSharedRef<FText>>)
+						.OptionsSource(&PipelineNames)
+						.OnGenerateWidget_Raw(this, &STakesView::OnGeneratePipelineNameWidget)
+						.OnSelectionChanged_Raw(this, &STakesView::OnPipelineSelectionChanged)
+						.Content()
 						[
-							RefreshButton.ToSharedRef()
+							SNew(STextBlock)
+							.MinDesiredWidth(200)
+							.Text_Lambda(
+								[this]()
+								{
+									return CurrentPipeline.DisplayName;
+								}
+							)
+							.ToolTipText_Lambda(
+								[this]()
+								{
+									return CurrentPipeline.ToolTip;
+								}
+							)
 						]
+					]
+					+ SHorizontalBox::Slot()
+					.Padding(2.0f)
+					.AutoWidth()
+					.VAlign(EVerticalAlignment::VAlign_Center)
+					[
+						AddToQueueButton.ToSharedRef()
+					]
+				]
+				+ SVerticalBox::Slot()
+				[
+					SNew(SOverlay)
+					+SOverlay::Slot()
+					[
+						ContentBrowser.CreateAssetPicker(TakesPickerConfig)
+					]
+					+SOverlay::Slot()
+					.VAlign(EVerticalAlignment::VAlign_Bottom)
+					.HAlign(EHorizontalAlignment::HAlign_Right)
+					[
+						SAssignNew(LoadingBox, SHorizontalBox)
+						.Visibility(EVisibility::Hidden)
+
 						+ SHorizontalBox::Slot()
-						.Padding(2.0f)
 						.AutoWidth()
 						.VAlign(EVerticalAlignment::VAlign_Center)
+						.HAlign(EHorizontalAlignment::HAlign_Center)
 						[
 							SNew(STextBlock)
-								.Text(JobPipelineText)
-						]
-						+ SHorizontalBox::Slot()
-						.Padding(2.0f)
-						.HAlign(EHorizontalAlignment::HAlign_Fill)
-						.VAlign(EVerticalAlignment::VAlign_Center)
-						[
-							SNew(SComboBox<TSharedRef<FText>>)
-								.OptionsSource(&PipelineNames)
-								.OnGenerateWidget_Raw(this, &STakesView::OnGeneratePipelineNameWidget)
-								.OnSelectionChanged_Raw(this, &STakesView::OnPipelineSelectionChanged)
-								.Content()
-								[
-									SNew(STextBlock)
-										.MinDesiredWidth(200)
-										.Text_Lambda(
-											[this]()
-											{
-												return CurrentPipeline.DisplayName;
-											}
-										)
-										.ToolTipText_Lambda(
-											[this]()
-											{
-												return CurrentPipeline.ToolTip;
-											}
-										)
-								]
+							.Text(LOCTEXT("LoadingTakesText", "Detecting takes..."))
 						]
 						+ SHorizontalBox::Slot()
-						.Padding(2.0f)
+						.Padding(4.0)
 						.AutoWidth()
 						.VAlign(EVerticalAlignment::VAlign_Center)
+						.HAlign(EHorizontalAlignment::HAlign_Center)
 						[
-							AddToQueueButton.ToSharedRef()
+							SNew(SCircularThrobber)
+							.Period(0.3f)
+							.Radius(10)
+							.PieceImage(FCoreStyle::Get().GetBrush("Throbber.CircleChunk.Small"))
 						]
-				]
-				+ SVerticalBox::Slot()
-				[
-					ContentBrowser.CreateAssetPicker(TakesPickerConfig)
+					]
 				]
 		];
 }
@@ -150,6 +215,18 @@ void STakesView::SetAddToQueueButtonEnabled(const bool bIsEnabled)
 	}
 }
 
+void STakesView::UpdateTakeListStarted()
+{
+	LoadingBox->SetVisibility(EVisibility::Visible);
+	RefreshButton->SetEnabled(false);
+}
+
+void STakesView::UpdateTakeListFinished()
+{
+	LoadingBox->SetVisibility(EVisibility::Hidden);
+	RefreshButton->SetEnabled(true);
+}
+
 TArray<TStrongObjectPtr<UTakeVirtualAsset>> STakesView::GetSelectedTakeAssets() const
 {
 	check(GetCurrentSelectionDelegate);
@@ -231,7 +308,7 @@ TArray<FAssetData> FTakesView::CreateTakeAssetsDataForCaptureDevice(TOptional<FS
 {
 	using namespace UE::CaptureManager;
 
-	TMap<FString, TArray<TStrongObjectPtr<UTakeVirtualAsset>>> SelectedDevicesTakesMap;
+	TArray<TStrongObjectPtr<UTakeVirtualAsset>> SelectedDevicesTakesArray;
 
 	TSharedPtr<FCaptureManagerPanelController> Controller = PanelController.Pin();
 	if (Controller)
@@ -249,27 +326,26 @@ TArray<FAssetData> FTakesView::CreateTakeAssetsDataForCaptureDevice(TOptional<FS
 				{
 					Takes.Add(Take);
 				}
-				SelectedDevicesTakesMap.Add(Device.GetName(), MoveTemp(Takes));
+
+				SelectedDevicesTakesArray.Append(MoveTemp(Takes));
 			}
 		}
 	}
 
-	return CreateTakeAssetsData(SelectedDevicesTakesMap);
+	return CreateTakeAssetsData(SelectedDevicesTakesArray);
 }
 
-TArray<FAssetData> FTakesView::CreateTakeAssetsData(const TMap<FString, TArray<TStrongObjectPtr<UTakeVirtualAsset>>>& InSelectedCaptureDevicesTakesMap)
+TArray<FAssetData> FTakesView::CreateTakeAssetsData(const TArray<TStrongObjectPtr<UTakeVirtualAsset>>& InSelectedCaptureDevicesTakesArray)
 {
 	TArray<FAssetData> TakeAssetsData;
-	for (const TTuple<FString, TArray<TStrongObjectPtr<UTakeVirtualAsset>>>& CaptureDeviceTakesTuple : InSelectedCaptureDevicesTakesMap)
+	for (TStrongObjectPtr<UTakeVirtualAsset> TakeItem : InSelectedCaptureDevicesTakesArray)
 	{
-		for (TStrongObjectPtr<UTakeVirtualAsset> TakeItem : CaptureDeviceTakesTuple.Value)
-		{
-			FAssetData TakeAssetData(TakeItem.Get());
-			TakeAssetData.AssetClassPath = FTopLevelAssetPath(*(TakeItem->GetPackage()->GetName()), *TakeItem->GetName());
+		FAssetData TakeAssetData(TakeItem.Get());
+		TakeAssetData.AssetClassPath = FTopLevelAssetPath(*(TakeItem->GetPackage()->GetName()), *TakeItem->GetName());
 
-			TakeAssetsData.Add(TakeAssetData);
-		}
+		TakeAssetsData.Add(TakeAssetData);
 	}
+
 	return TakeAssetsData;
 }
 
@@ -358,26 +434,69 @@ void FTakesView::RefreshTakes()
 	{
 		if (SelectedItem->Implements<ULiveLinkDeviceCapability_Ingest>())
 		{
-			FIngestUpdateTakeListCallback Callback = FIngestUpdateTakeListCallback::Type::CreateLambda(
-				[this, InCaptureDeviceId = SelectedItem->GetDeviceId()](const TArray<FTakeId>& InTakeIds)
-				{
-					CaptureDeviceTakesMap.Remove(InCaptureDeviceId);
+			FGuid DeviceId = SelectedItem->GetDeviceId();
 
-					TArray<TStrongObjectPtr<UTakeVirtualAsset>> TakeObjects = MakeTakeObjects(InCaptureDeviceId, InTakeIds);
+			ResetTakesCache(DeviceId);
 
-					CaptureDeviceTakesMap.Add(InCaptureDeviceId, MoveTemp(TakeObjects));
-
-					RefreshTakesViewDelegate.ExecuteIfBound(true);
-				}
-			);
+			FIngestUpdateTakeListCallback Callback =
+				FIngestUpdateTakeListCallback::Type::CreateRaw(this, &FTakesView::UpdateTakeListCallback, DeviceId);
 
 			UIngestCapability_UpdateTakeListCallback* UpdateTakeListCallback = NewObject<UIngestCapability_UpdateTakeListCallback>();
 			UpdateTakeListCallback->Callback = MoveTemp(Callback);
+
+			TakesTileView->UpdateTakeListStarted();
 			ILiveLinkDeviceCapability_Ingest::Execute_UpdateTakeList(SelectedItem.Get(), UpdateTakeListCallback);
 		}
 	}
 }
 
+void FTakesView::UpdateTakeListCallback(TArray<UE::CaptureManager::FTakeId> InTakeIds, FGuid InCaptureDeviceId)
+{
+	using namespace UE::CaptureManager;
+
+	check(IsInGameThread());
+
+	TakesTileView->UpdateTakeListFinished();
+
+	TArray<TStrongObjectPtr<UTakeVirtualAsset>>* Found = CaptureDeviceTakesMap.Find(InCaptureDeviceId);
+	if (!Found)
+	{
+		return;
+	}
+
+	const TArray<TStrongObjectPtr<UTakeVirtualAsset>>& TakeAssets = *Found;
+
+	if (TakeAssets.IsEmpty() && !InTakeIds.IsEmpty())
+	{
+		TArray<TStrongObjectPtr<UTakeVirtualAsset>> TakeObjects = MakeTakeObjects(InCaptureDeviceId, InTakeIds);
+
+		CaptureDeviceTakesMap.Add(InCaptureDeviceId, MoveTemp(TakeObjects));
+
+		RefreshTakesViewDelegate.ExecuteIfBound(true);
+		return;
+	}
+
+	bool bAreEqual = EqualTo(TakeAssets, InTakeIds, 
+							 [](const TStrongObjectPtr<UTakeVirtualAsset>& InLeft, FTakeId InTakeId)
+							 {
+								 return InLeft->TakeId == InTakeId;
+							 });
+
+	check(bAreEqual);
+}
+
+void FTakesView::ResetTakesCache(FGuid InCaptureDeviceId)
+{
+	if (CaptureDeviceTakesMap.Contains(InCaptureDeviceId))
+	{
+		CaptureDeviceTakesMap[InCaptureDeviceId].Empty();
+	}
+	else
+	{
+		CaptureDeviceTakesMap.Add(InCaptureDeviceId);
+	}
+}
+
 void FTakesView::CreateTakesView()
 {
 	TSharedPtr<FCaptureManagerPanelController> Controller = PanelController.Pin();
@@ -468,19 +587,6 @@ void FTakesView::CaptureDeviceStarted(FGuid InCaptureDeviceId)
 			FString CaptureDeviceName = Device->GetName();
 			FName Name(CaptureDeviceName);
 
-			FIngestUpdateTakeListCallback Callback = FIngestUpdateTakeListCallback::Type::CreateLambda(
-				[this, InCaptureDeviceId](const TArray<FTakeId>& InTakeIds)
-				{
-					CaptureDeviceTakesMap.Remove(InCaptureDeviceId);
-
-					TArray<TStrongObjectPtr<UTakeVirtualAsset>> TakeObjects = MakeTakeObjects(InCaptureDeviceId, InTakeIds);
-
-					CaptureDeviceTakesMap.Add(InCaptureDeviceId, MoveTemp(TakeObjects));
-
-					RefreshTakesViewDelegate.ExecuteIfBound(true);
-				}
-			);
-
 			TScriptInterface<ILiveLinkDeviceCapability_Ingest> CapabilityInterface(Device);
 
 			CapabilityInterface->SubscribeToEvent(FString(FIngestCapability_TakeAddedEvent::Name),
@@ -551,8 +657,16 @@ void FTakesView::CaptureDeviceStarted(FGuid InCaptureDeviceId)
 						}
 					}));
 
+			ResetTakesCache(InCaptureDeviceId);
+
+			FIngestUpdateTakeListCallback Callback = 
+				FIngestUpdateTakeListCallback::Type::CreateRaw(this, &FTakesView::UpdateTakeListCallback, InCaptureDeviceId);
+
 			UIngestCapability_UpdateTakeListCallback* UpdateTakeListCallback = NewObject<UIngestCapability_UpdateTakeListCallback>();
 			UpdateTakeListCallback->Callback = MoveTemp(Callback);
+
+			TakesTileView->UpdateTakeListStarted();
+
 			ILiveLinkDeviceCapability_Ingest::Execute_UpdateTakeList(Device, UpdateTakeListCallback);
 		});
 }
@@ -589,7 +703,8 @@ void FTakesView::CaptureDeviceToBeRemoved(FGuid InCaptureDeviceId)
 		});
 }
 
-TArray<TStrongObjectPtr<UTakeVirtualAsset>> FTakesView::MakeTakeObjects(FGuid InCaptureDeviceId, TOptional<TArray<UE::CaptureManager::FTakeId>> InTakeIds) {
+TArray<TStrongObjectPtr<UTakeVirtualAsset>> FTakesView::MakeTakeObjects(FGuid InCaptureDeviceId, TOptional<TArray<UE::CaptureManager::FTakeId>> InTakeIds) 
+{
 	TArray<TStrongObjectPtr<UTakeVirtualAsset>> TakeObjects;
 
 	TSharedPtr<FCaptureManagerPanelController> Controller = PanelController.Pin();
diff --git a/Engine/Plugins/VirtualProduction/CaptureManager/CaptureManagerApp/Source/CaptureManagerEditor/Private/UI/TakesView.h b/Engine/Plugins/VirtualProduction/CaptureManager/CaptureManagerApp/Source/CaptureManagerEditor/Private/UI/TakesView.h
--- a/Engine/Plugins/VirtualProduction/CaptureManager/CaptureManagerApp/Source/CaptureManagerEditor/Private/UI/TakesView.h
+++ b/Engine/Plugins/VirtualProduction/CaptureManager/CaptureManagerApp/Source/CaptureManagerEditor/Private/UI/TakesView.h
@@ -52,6 +52,8 @@ public:
 		TUniquePtr<FGetCurrentSelectionDelegate> GetCurrentSelectionDelegate
 	);
 	void SetAddToQueueButtonEnabled(bool bIsEnabled);
+	void UpdateTakeListStarted();
+	void UpdateTakeListFinished();
 
 private:
 	FReply OnRefreshButtonClicked();
@@ -65,6 +67,7 @@ private:
 	FOnRefreshTakes OnRefreshTakes;
 	TSharedPtr<SPositiveActionButton> AddToQueueButton;
 	FOnAddTakesToIngestQueue OnAddTakesToIngestQueue;
+	TSharedPtr<SHorizontalBox> LoadingBox;
 
 	TSharedPtr<UE::CaptureManager::FIngestPipelineManager> IngestPipelineManager;
 	TArray<TSharedRef<FText>> PipelineNames;
@@ -94,11 +97,13 @@ private:
 	FSyncToAssetsDelegate SyncToDevicesDelegate;
 
 	TArray<FAssetData> CreateTakeAssetsDataForCaptureDevice(TOptional<FString> InCaptureDeviceName = TOptional<FString>());
-	TArray<FAssetData> CreateTakeAssetsData(const TMap<FString, TArray<TStrongObjectPtr<UTakeVirtualAsset>>>& InSelectedCaptureDevicesTakesMap);
+	TArray<FAssetData> CreateTakeAssetsData(const TArray<TStrongObjectPtr<UTakeVirtualAsset>>& InSelectedCaptureDevicesTakesArray);
 
 	EImageRotation ConvertOrientation(TOptional<FTakeMetadata::FVideo::EOrientation> InOrientation) const;
 	void AddTakesToIngestQueue(const TArray<TStrongObjectPtr<UTakeVirtualAsset>>& InTakeAssets, const UIngestJobSettingsBase& InJobPreset);
 	void RefreshTakes();
+	void UpdateTakeListCallback(TArray<UE::CaptureManager::FTakeId> InTakeIds, FGuid InCaptureDeviceId);
+	void ResetTakesCache(FGuid InCaptureDeviceId);
 
 	TArray<TStrongObjectPtr<UTakeVirtualAsset>> MakeTakeObjects(FGuid InCaptureDeviceId,
 																TOptional<TArray<UE::CaptureManager::FTakeId>> InTakeIds = TOptional<TArray<UE::CaptureManager::FTakeId>>());
diff --git a/Engine/Plugins/VirtualProduction/CaptureManager/CaptureManagerDevices/Source/FilesystemLiveLinkDevice/Private/FilesystemLiveLinkDevice.cpp b/Engine/Plugins/VirtualProduction/CaptureManager/CaptureManagerDevices/Source/FilesystemLiveLinkDevice/Private/FilesystemLiveLinkDevice.cpp
--- a/Engine/Plugins/VirtualProduction/CaptureManager/CaptureManagerDevices/Source/FilesystemLiveLinkDevice/Private/FilesystemLiveLinkDevice.cpp
+++ b/Engine/Plugins/VirtualProduction/CaptureManager/CaptureManagerDevices/Source/FilesystemLiveLinkDevice/Private/FilesystemLiveLinkDevice.cpp
@@ -127,6 +127,8 @@ void UFilesystemLiveLinkDevice::UpdateTakeList_Implementation(UIngestCapability_
 				FString TakePath = CurrentDirectory.Right(CurrentDirectory.Len() - StoragePath.Len());
 				FPaths::NormalizeDirectoryName(TakePath); // Removes trailing slash
 				RelativeTakePaths.Add(TakeId, MoveTemp(TakePath));
+
+				PublishEvent<FIngestCapability_TakeAddedEvent>(TakeId);
 			}
 		}
 
diff --git a/Engine/Plugins/VirtualProduction/CaptureManager/CaptureManagerDevices/Source/MonoVideoLiveLinkDevice/Private/MonoVideoLiveLinkDevice.cpp b/Engine/Plugins/VirtualProduction/CaptureManager/CaptureManagerDevices/Source/MonoVideoLiveLinkDevice/Private/MonoVideoLiveLinkDevice.cpp
--- a/Engine/Plugins/VirtualProduction/CaptureManager/CaptureManagerDevices/Source/MonoVideoLiveLinkDevice/Private/MonoVideoLiveLinkDevice.cpp
+++ b/Engine/Plugins/VirtualProduction/CaptureManager/CaptureManagerDevices/Source/MonoVideoLiveLinkDevice/Private/MonoVideoLiveLinkDevice.cpp
@@ -118,48 +118,53 @@ void UMonoVideoLiveLinkDevice::RunIngestTake(int32 InTakeId, const UIngestCapabi
 
 void UMonoVideoLiveLinkDevice::UpdateTakeList_Implementation(UIngestCapability_UpdateTakeListCallback* InCallback)
 {
-	RemoveAllTakes();
+	AsyncTask(ENamedThreads::AnyThread, [this, Callback = TStrongObjectPtr<UIngestCapability_UpdateTakeListCallback>(InCallback)]
+	{
+		RemoveAllTakes();
 
-	FString StoragePath = GetSettings()->TakeDirectory.Path;
-	int32 PathsInterrogatedCount = 0;
-	const int32 PathsToInterrogateInOneRun = 100000;
+		FString StoragePath = GetSettings()->TakeDirectory.Path;
+		int32 PathsInterrogatedCount = 0;
+		const int32 PathsToInterrogateInOneRun = 100000;
 	
-	TArray<FString> SupportedVideoFiles;
-	const bool bIterationResult = IFileManager::Get().IterateDirectoryRecursively(*StoragePath, [&](const TCHAR* InFileNameOrDirectory, bool bInIsDirectory)
-	{
-		if (++PathsInterrogatedCount > PathsToInterrogateInOneRun)
+		TArray<FString> SupportedVideoFiles;
+		const bool bIterationResult = IFileManager::Get().IterateDirectoryRecursively(*StoragePath, [&](const TCHAR* InFileNameOrDirectory, bool bInIsDirectory)
 		{
-			EAppReturnType::Type Result = FMessageDialog::Open(EAppMsgType::OkCancel, FText::Format(LOCTEXT("MonoVideoDevice_ProceedWithTakesDiscovery", "Analyzed {0} paths in search for takes. Do you want to analyze {1} more?"), PathsToInterrogateInOneRun, PathsToInterrogateInOneRun));
-			if (Result == EAppReturnType::Ok)
+			if (++PathsInterrogatedCount > PathsToInterrogateInOneRun)
 			{
-				PathsInterrogatedCount = 0;
+				EAppReturnType::Type Result = FMessageDialog::Open(EAppMsgType::OkCancel, FText::Format(LOCTEXT("MonoVideoDevice_ProceedWithTakesDiscovery", "Analyzed {0} paths in search for takes. Do you want to analyze {1} more?"), PathsToInterrogateInOneRun, PathsToInterrogateInOneRun));
+				if (Result == EAppReturnType::Ok)
+				{
+					PathsInterrogatedCount = 0;
+				}
+				else 
+				{
+					return false;
+				}
 			}
-			else 
-			{
-				return false;
+
+			if (!bInIsDirectory && IsVideoFile(InFileNameOrDirectory))
+			{			
+				SupportedVideoFiles.Add(InFileNameOrDirectory);
 			}
-		}
 
-		if (!bInIsDirectory && IsVideoFile(InFileNameOrDirectory))
-		{			
-			SupportedVideoFiles.Add(InFileNameOrDirectory);
-		}
+			return true;
+		});
 
-		return true;
-	});
+		for (const FString& SupportedVideoFile : SupportedVideoFiles)
+		{
+			TOptional<FTakeMetadata> TakeInfoResult = ReadTake(SupportedVideoFile);
 
-	for (const FString& SupportedVideoFile : SupportedVideoFiles)
-	{
-		TOptional<FTakeMetadata> TakeInfoResult = ReadTake(SupportedVideoFile);
+			if (TakeInfoResult.IsSet())
+			{
+				int32 TakeId = AddTake(TakeInfoResult.GetValue());
+				FullTakePaths.Add(TakeId, SupportedVideoFile);
 
-		if (TakeInfoResult.IsSet())
-		{
-			int32 TakeId = AddTake(TakeInfoResult.GetValue());
-			FullTakePaths.Add(TakeId, SupportedVideoFile);
+				PublishEvent<FIngestCapability_TakeAddedEvent>(TakeId);
+			}
 		}
-	}
 
-	ExecuteUpdateTakeListCallback(InCallback, Execute_GetTakeIdentifiers(this));
+		ExecuteUpdateTakeListCallback(Callback.Get(), Execute_GetTakeIdentifiers(this));
+	});
 }
 
 ELiveLinkDeviceConnectionStatus UMonoVideoLiveLinkDevice::GetConnectionStatus_Implementation() const
diff --git a/Engine/Plugins/VirtualProduction/CaptureManager/CaptureManagerDevices/Source/StereoVideoLiveLinkDevice/Private/StereoVideoLiveLinkDevice.cpp b/Engine/Plugins/VirtualProduction/CaptureManager/CaptureManagerDevices/Source/StereoVideoLiveLinkDevice/Private/StereoVideoLiveLinkDevice.cpp
--- a/Engine/Plugins/VirtualProduction/CaptureManager/CaptureManagerDevices/Source/StereoVideoLiveLinkDevice/Private/StereoVideoLiveLinkDevice.cpp
+++ b/Engine/Plugins/VirtualProduction/CaptureManager/CaptureManagerDevices/Source/StereoVideoLiveLinkDevice/Private/StereoVideoLiveLinkDevice.cpp
@@ -144,41 +144,46 @@ int32 UStereoVideoLiveLinkDevice::FTakeWithComponents::CountComponents(UStereoVi
 
 void UStereoVideoLiveLinkDevice::UpdateTakeList_Implementation(UIngestCapability_UpdateTakeListCallback* InCallback)
 {
-	RemoveAllTakes();
+	AsyncTask(ENamedThreads::AnyThread, [this, Callback = TStrongObjectPtr<UIngestCapability_UpdateTakeListCallback>(InCallback)]()
+	{
+		RemoveAllTakes();
 
-	FString StoragePath = GetSettings()->TakeDirectory.Path;
+		FString StoragePath = GetSettings()->TakeDirectory.Path;
 
-	TMap<FString, UStereoVideoLiveLinkDevice::FTakeWithComponents> TakesWithComponentsCandidates = DiscoverTakes(StoragePath);
+		TMap<FString, UStereoVideoLiveLinkDevice::FTakeWithComponents> TakesWithComponentsCandidates = DiscoverTakes(StoragePath);
 
-	TArray<FString> SlateNames;
-	TakesWithComponentsCandidates.GetKeys(SlateNames);
+		TArray<FString> SlateNames;
+		TakesWithComponentsCandidates.GetKeys(SlateNames);
 
-	for (const FString& SlateName : SlateNames)
-	{
-		FTakeWithComponents& TakeWithComponents = TakesWithComponentsCandidates[SlateName];
+		for (const FString& SlateName : SlateNames)
+		{
+			FTakeWithComponents& TakeWithComponents = TakesWithComponentsCandidates[SlateName];
 
-		bool TakeAsExpected = ((TakeWithComponents.CountComponents(ETakeComponentType::VIDEO) == 2 && TakeWithComponents.CountComponents(ETakeComponentType::IMAGE_SEQUENCE) == 0) ||
-			(TakeWithComponents.CountComponents(ETakeComponentType::VIDEO) == 0 && TakeWithComponents.CountComponents(ETakeComponentType::IMAGE_SEQUENCE) == 2)) &&
-			TakeWithComponents.CountComponents(ETakeComponentType::AUDIO) <= 1;
+			bool TakeAsExpected = ((TakeWithComponents.CountComponents(ETakeComponentType::VIDEO) == 2 && TakeWithComponents.CountComponents(ETakeComponentType::IMAGE_SEQUENCE) == 0) ||
+									(TakeWithComponents.CountComponents(ETakeComponentType::VIDEO) == 0 && TakeWithComponents.CountComponents(ETakeComponentType::IMAGE_SEQUENCE) == 2)) &&
+				TakeWithComponents.CountComponents(ETakeComponentType::AUDIO) <= 1;
 
-		if (!TakeAsExpected)
-		{
-			UE_LOG(LogVideoLiveLinkDevice, Warning, TEXT("Invalid take at '%s'. Take should have exactly two video files or two image sequences. Optionally, one audio file."), *TakeWithComponents.TakePath);
-			continue;
-		}
+			if (!TakeAsExpected)
+			{
+				UE_LOG(LogVideoLiveLinkDevice, Warning, TEXT("Invalid take at '%s'. Take should have exactly two video files or two image sequences. Optionally, one audio file."), *TakeWithComponents.TakePath);
+				continue;
+			}
 
-		TOptional<FTakeMetadata> TakeInfoResult = CreateTakeMetadata(TakeWithComponents);
+			TOptional<FTakeMetadata> TakeInfoResult = CreateTakeMetadata(TakeWithComponents);
 
-		if (TakeInfoResult.IsSet())
-		{
-			FTakeMetadata TakeMetadata = TakeInfoResult.GetValue();
+			if (TakeInfoResult.IsSet())
+			{
+				FTakeMetadata TakeMetadata = TakeInfoResult.GetValue();
 
-			int32 TakeId = AddTake(MoveTemp(TakeMetadata));
-			FullTakePaths.Add(TakeId, TakeWithComponents.TakePath);
+				int32 TakeId = AddTake(MoveTemp(TakeMetadata));
+				FullTakePaths.Add(TakeId, TakeWithComponents.TakePath);
+
+				PublishEvent<FIngestCapability_TakeAddedEvent>(TakeId);
+			}
 		}
-	}
 
-	ExecuteUpdateTakeListCallback(InCallback, Execute_GetTakeIdentifiers(this));
+		ExecuteUpdateTakeListCallback(Callback.Get(), Execute_GetTakeIdentifiers(this));
+	});
 }
 
 TMap<FString, UStereoVideoLiveLinkDevice::FTakeWithComponents> UStereoVideoLiveLinkDevice::DiscoverTakes(FString StoragePath)
