Thumbnail extraction with more formats supported

diff --git a/Engine/Plugins/VirtualProduction/CaptureManager/CaptureManagerApp/Source/CaptureManagerEditor/Private/UI/TakesView.cpp b/Engine/Plugins/VirtualProduction/CaptureManager/CaptureManagerApp/Source/CaptureManagerEditor/Private/UI/TakesView.cpp
--- a/Engine/Plugins/VirtualProduction/CaptureManager/CaptureManagerApp/Source/CaptureManagerEditor/Private/UI/TakesView.cpp
+++ b/Engine/Plugins/VirtualProduction/CaptureManager/CaptureManagerApp/Source/CaptureManagerEditor/Private/UI/TakesView.cpp
@@ -777,19 +777,27 @@ TArray<TStrongObjectPtr<UTakeVirtualAsset>> FTakesView::MakeTakeObjects(FGuid In
 TOptional<FTakeThumbnail> FTakesView::CreateThumbnail(const FTakeThumbnailData& InThumbnailData)
 {
 	TOptional<TArray<uint8>> ThumbnailRawDataOpt = InThumbnailData.GetThumbnailData();
+	TOptional<FTakeThumbnailData::FRawImage> ThumbnailRawImageOpt = InThumbnailData.GetRawImage();
 
-	if (!ThumbnailRawDataOpt.IsSet())
+	UTexture2D* PreviewImageTexture = nullptr; 
+	if (ThumbnailRawDataOpt.IsSet())
 	{
-		return {};
+		FImage Image;
+		TArray<uint8> ThumbnailRawData = MoveTemp(ThumbnailRawDataOpt.GetValue());
+		FImageUtils::DecompressImage(ThumbnailRawData.GetData(), ThumbnailRawData.Num(), Image);
+		Image.ChangeFormat(ERawImageFormat::BGRA8, EGammaSpace::sRGB);
+
+		PreviewImageTexture = FImageUtils::CreateTexture2DFromImage(Image);
 	}
+	else if (ThumbnailRawImageOpt.IsSet())
+	{
+		FTakeThumbnailData::FRawImage RawImage = MoveTemp(ThumbnailRawImageOpt.GetValue());
 
-	TArray<uint8> ThumbnailRawData = MoveTemp(ThumbnailRawDataOpt.GetValue());
+		FImageView ImageView(RawImage.DecompressedImageData.GetData(), RawImage.Width, RawImage.Height);
 
-	FImage Image;
-	FImageUtils::DecompressImage(ThumbnailRawData.GetData(), ThumbnailRawData.Num(), Image);
-	Image.ChangeFormat(ERawImageFormat::BGRA8, EGammaSpace::sRGB);
+		PreviewImageTexture = FImageUtils::CreateTexture2DFromImage(ImageView);
+	}
 
-	UTexture2D* PreviewImageTexture = FImageUtils::CreateTexture2DFromImage(Image);
 	if (!PreviewImageTexture)
 	{
 		return {};
diff --git a/Engine/Plugins/VirtualProduction/CaptureManager/CaptureManagerApp/Source/CaptureManagerMediaRW/Private/Utils/MediaPixelFormatConversions.cpp b/Engine/Plugins/VirtualProduction/CaptureManager/CaptureManagerApp/Source/CaptureManagerMediaRW/Private/Utils/MediaPixelFormatConversions.cpp
new file mode 100644
--- /dev/null
+++ b/Engine/Plugins/VirtualProduction/CaptureManager/CaptureManagerApp/Source/CaptureManagerMediaRW/Private/Utils/MediaPixelFormatConversions.cpp
@@ -0,0 +1,441 @@
+// Copyright Epic Games, Inc. All Rights Reserved.
+
+#include "Utils/MediaPixelFormatConversions.h"
+
+#include "Async/ParallelFor.h"
+
+namespace UE::CaptureManager
+{
+
+namespace Private
+{
+template<typename T>
+static T ScalePixel(T InPixel, TRange<T> InInputRange, TRange<T> InOutputRange)
+{
+	const float Factor = static_cast<float>(InOutputRange.GetUpperBoundValue()) / (InInputRange.GetUpperBoundValue() - InInputRange.GetLowerBoundValue());
+	const uint64 NewPixelValue = Factor * (InPixel - InInputRange.GetLowerBoundValue()) + 0.5;
+	return FMath::Clamp(NewPixelValue, InOutputRange.GetLowerBoundValue(), InOutputRange.GetUpperBoundValue());
+}
+
+static inline void RGBfromYUV(uint8& R, uint8& G, uint8& B, uint8 Y, uint8 U, uint8 V)
+{
+	int64 BigY = Y - 16;
+	int64 BigU = U - 128;
+	int64 BigV = V - 128;
+
+	int64 BigR = (298 * BigY + 409 * BigV + 128) >> 8;
+	int64 BigG = (298 * BigY - 100 * BigU - 208 * BigV + 128) >> 8;
+	int64 BigB = (298 * BigY + 516 * BigU + 128) >> 8;
+
+	R = FMath::Clamp(BigR, 0, 255);
+	G = FMath::Clamp(BigG, 0, 255);
+	B = FMath::Clamp(BigB, 0, 255);
+}
+
+static inline FColor CreateColor(uint8 R, uint8 G, uint8 B, uint8 A)
+{
+	return FColor(R, G, B, A);
+}
+
+}
+
+TArray<uint8> ConvertYUVToMono(const UE::CaptureManager::FMediaTextureSample* InSample, bool InScaleRange)
+{
+	using namespace UE::CaptureManager;
+
+	TArray<uint8> NewBuffer;
+	NewBuffer.SetNumUninitialized(InSample->Stride * InSample->Dimensions.Y);
+
+	uint32 StrideY = InSample->Stride;
+	int32 MonoStride = InSample->Stride;
+
+	ParallelFor(InSample->Dimensions.Y, [&NewBuffer, StrideY, MonoStride, InScaleRange, InSample](int32 Y) mutable
+	{
+		for (int32 X = 0; X < InSample->Dimensions.X; ++X)
+		{
+			int32 YIndex = Y * StrideY + X;
+			uint8 PixelValue = InSample->Buffer[YIndex];
+
+			if (InScaleRange)
+			{
+				using FRange = TRange<uint8>;
+				PixelValue = Private::ScalePixel(PixelValue,
+												 FRange::Inclusive(16, 235),
+												 FRange::Inclusive(MIN_uint8, MAX_uint8));
+			}
+
+			NewBuffer[Y * MonoStride + X] = PixelValue;
+		}
+	});
+
+	return NewBuffer;
+}
+
+TArray<uint8> ConvertYUY2ToMono(const UE::CaptureManager::FMediaTextureSample* InSample, bool InScaleRange)
+{
+	using namespace UE::CaptureManager;
+
+	TArray<uint8> NewBuffer;
+	NewBuffer.SetNumUninitialized(InSample->Stride * InSample->Dimensions.Y);
+
+	int32 YUY2Stride = InSample->Stride * 2;
+	int32 MonoStride = InSample->Stride;
+
+	ParallelFor(InSample->Dimensions.Y, [&NewBuffer, YUY2Stride, MonoStride, InScaleRange, InSample](int32 Y) mutable
+	{
+		for (int32 X = 0; X < InSample->Dimensions.X; X += 2)
+		{
+			int32 YIndex = Y * YUY2Stride + X * 2;
+
+			uint8_t Y0 = InSample->Buffer[YIndex];
+			uint8_t Y1 = InSample->Buffer[YIndex + 2];
+
+			if (InScaleRange)
+			{
+				using FRange = TRange<uint8>;
+				Y0 = Private::ScalePixel(Y0,
+										 FRange::Inclusive(16, 235),
+										 FRange::Inclusive(MIN_uint8, MAX_uint8));
+				Y1 = Private::ScalePixel(Y1,
+										 FRange::Inclusive(16, 235),
+										 FRange::Inclusive(MIN_uint8, MAX_uint8));
+
+			}
+
+			NewBuffer[Y * MonoStride + X] = Y0;
+			NewBuffer[Y * MonoStride + X + 1] = Y1;
+		}
+	});
+
+	return NewBuffer;
+}
+
+TArray<uint8> ConvertI420ToBGRA(const UE::CaptureManager::FMediaTextureSample* InSample)
+{
+	using namespace UE::CaptureManager;
+
+	EMediaTexturePixelFormat OutputFormat = EMediaTexturePixelFormat::U8_BGRA;
+	uint32 Channels = GetNumberOfChannels(OutputFormat);
+
+	TArray<uint8> NewBuffer;
+	NewBuffer.SetNumUninitialized(InSample->Stride * InSample->Dimensions.Y * Channels);
+
+	uint32 StrideY = InSample->Stride;
+	uint32 StrideUV = InSample->Stride / 2;
+	int32 BGRAStride = InSample->Stride * Channels;
+
+	uint32 YPlaneSize = StrideY * InSample->Dimensions.Y;
+	uint32 UPlaneSize = StrideUV * (InSample->Dimensions.Y / 2);
+	uint32 VPlaneSize = InSample->Buffer.Num() - (YPlaneSize + UPlaneSize);
+
+	const TArrayView<const uint8> YPlane = MakeArrayView(InSample->Buffer.GetData(), YPlaneSize);
+	const TArrayView<const uint8> UPlane = MakeArrayView(InSample->Buffer.GetData() + YPlaneSize, UPlaneSize);
+	const TArrayView<const uint8> VPlane = MakeArrayView(InSample->Buffer.GetData() + YPlaneSize + UPlaneSize, VPlaneSize);
+
+	ParallelFor(InSample->Dimensions.Y, [&NewBuffer, YPlane, UPlane, VPlane, StrideY, StrideUV, BGRAStride, InSample](int32 Y) mutable
+	{
+		for (int32 X = 0; X < InSample->Dimensions.X; ++X)
+		{
+			int32 YIndex = Y * StrideY + X;
+			int32 UVIndex = (Y / 2) * StrideUV + (X / 2);
+
+			uint8_t Y0 = YPlane[YIndex];
+			uint8_t U = UPlane[UVIndex];
+			uint8_t V = VPlane[UVIndex];
+
+			uint8 R, G, B;
+			Private::RGBfromYUV(R, G, B, Y0, U, V);
+
+			NewBuffer[Y * BGRAStride + X * 4] = B;
+			NewBuffer[Y * BGRAStride + X * 4 + 1] = G;
+			NewBuffer[Y * BGRAStride + X * 4 + 2] = R;
+			NewBuffer[Y * BGRAStride + X * 4 + 3] = 255;
+		}
+	});
+
+	return NewBuffer;
+}
+
+TArray<uint8> ConvertNV12ToBGRA(const UE::CaptureManager::FMediaTextureSample* InSample)
+{
+	using namespace UE::CaptureManager;
+
+	EMediaTexturePixelFormat OutputFormat = EMediaTexturePixelFormat::U8_BGRA;
+	uint32 Channels = GetNumberOfChannels(OutputFormat);
+
+	TArray<uint8> NewBuffer;
+	NewBuffer.SetNumUninitialized(InSample->Stride * InSample->Dimensions.Y * Channels);
+
+	uint32 StrideY = InSample->Stride;
+	int32 BGRAStride = InSample->Stride * Channels;
+
+	uint32 YPlaneSize = StrideY * InSample->Dimensions.Y;
+	uint32 UVPlaneSize = InSample->Buffer.Num() - YPlaneSize;
+
+	const TArrayView<const uint8> YPlane = MakeArrayView(InSample->Buffer.GetData(), YPlaneSize);
+	const TArrayView<const uint8> UVPlane = MakeArrayView(InSample->Buffer.GetData() + YPlaneSize, UVPlaneSize);
+
+	ParallelFor(InSample->Dimensions.Y, [&NewBuffer, YPlane, UVPlane, StrideY, BGRAStride, InSample, Channels](int32 Y) mutable
+	{
+		for (int32 X = 0; X < InSample->Dimensions.X; ++X)
+		{
+			int32 YIndex = Y * StrideY + X;
+			int32 UVIndex = (Y / 2) * StrideY + (X & ~1);
+
+			uint8_t Y0 = YPlane[YIndex];
+			uint8_t U = UVPlane[UVIndex];
+			uint8_t V = UVPlane[UVIndex + 1];
+
+			uint8 R, G, B;
+			Private::RGBfromYUV(R, G, B, Y0, U, V);
+
+			NewBuffer[Y * BGRAStride + X * 4] = B;
+			NewBuffer[Y * BGRAStride + X * 4 + 1] = G;
+			NewBuffer[Y * BGRAStride + X * 4 + 2] = R;
+			NewBuffer[Y * BGRAStride + X * 4 + 3] = 255;
+		}
+	});
+
+	return NewBuffer;
+}
+
+TArray<uint8> ConvertYUY2ToBGRA(const UE::CaptureManager::FMediaTextureSample* InSample)
+{
+	using namespace UE::CaptureManager;
+
+	EMediaTexturePixelFormat OutputFormat = EMediaTexturePixelFormat::U8_BGRA;
+
+	uint32 Channels = GetNumberOfChannels(OutputFormat);
+
+	TArray<uint8> NewBuffer;
+	NewBuffer.SetNumUninitialized(InSample->Stride * InSample->Dimensions.Y * Channels);
+
+	int32 YUY2Stride = InSample->Stride * 2;
+	int32 BGRAStride = InSample->Stride * Channels;
+
+	ParallelFor(InSample->Dimensions.Y, [&NewBuffer, YUY2Stride, BGRAStride, InSample](int32 Y) mutable
+	{
+		for (int32 X = 0; X < InSample->Dimensions.X; X += 2)
+		{
+			int32 YIndex = Y * YUY2Stride + X * 2;
+
+			uint8 Y0 = InSample->Buffer[YIndex];
+			uint8 U = InSample->Buffer[YIndex + 1];
+			uint8 Y1 = InSample->Buffer[YIndex + 2];
+			uint8 V = InSample->Buffer[YIndex + 3];
+
+			uint8 R0, G0, B0;
+			Private::RGBfromYUV(R0, G0, B0, Y0, U, V);
+
+			uint8 R1, G1, B1;
+			Private::RGBfromYUV(R1, G1, B1, Y1, U, V);
+
+			NewBuffer[Y * BGRAStride + X * 4] = B0;
+			NewBuffer[Y * BGRAStride + X * 4 + 1] = G0;
+			NewBuffer[Y * BGRAStride + X * 4 + 2] = R0;
+			NewBuffer[Y * BGRAStride + X * 4 + 3] = 255;
+
+			NewBuffer[Y * BGRAStride + (X + 1) * 4] = B1;
+			NewBuffer[Y * BGRAStride + (X + 1) * 4 + 1] = G1;
+			NewBuffer[Y * BGRAStride + (X + 1) * 4 + 2] = R1;
+			NewBuffer[Y * BGRAStride + (X + 1) * 4 + 3] = 255;
+		}
+	});
+
+	return NewBuffer;
+}
+
+TArray<FColor> UEConvertYUVToMono(const UE::CaptureManager::FMediaTextureSample* InSample, bool InScaleRange)
+{
+	using namespace UE::CaptureManager;
+
+	TArray<FColor> NewBuffer;
+	NewBuffer.SetNumUninitialized(InSample->Stride * InSample->Dimensions.Y);
+
+	uint32 StrideY = InSample->Stride;
+	int32 MonoStride = InSample->Stride;
+
+	ParallelFor(InSample->Dimensions.Y, [&NewBuffer, StrideY, MonoStride, InScaleRange, InSample](int32 Y) mutable
+	{
+		for (int32 X = 0; X < InSample->Dimensions.X; ++X)
+		{
+			int32 YIndex = Y * StrideY + X;
+			uint8 PixelValue = InSample->Buffer[YIndex];
+
+			if (InScaleRange)
+			{
+				using FRange = TRange<uint8>;
+				PixelValue = Private::ScalePixel(PixelValue,
+												 FRange::Inclusive(16, 235),
+												 FRange::Inclusive(MIN_uint8, MAX_uint8));
+			}
+
+
+			NewBuffer[Y * MonoStride + X] = 
+				Private::CreateColor(PixelValue, PixelValue, PixelValue, InScaleRange ? 235 : 255);
+		}
+	});
+
+	return NewBuffer;
+}
+
+TArray<FColor> UEConvertYUY2ToMono(const UE::CaptureManager::FMediaTextureSample* InSample, bool InScaleRange)
+{
+	using namespace UE::CaptureManager;
+
+	TArray<FColor> NewBuffer;
+	NewBuffer.SetNumUninitialized(InSample->Stride * InSample->Dimensions.Y);
+
+	int32 YUY2Stride = InSample->Stride * 2;
+	int32 MonoStride = InSample->Stride;
+
+	ParallelFor(InSample->Dimensions.Y, [&NewBuffer, YUY2Stride, MonoStride, InScaleRange, InSample](int32 Y) mutable
+	{
+		for (int32 X = 0; X < InSample->Dimensions.X; X += 2)
+		{
+			int32 YIndex = Y * YUY2Stride + X * 2;
+
+			uint8_t Y0 = InSample->Buffer[YIndex];
+			uint8_t Y1 = InSample->Buffer[YIndex + 2];
+
+			if (InScaleRange)
+			{
+				using FRange = TRange<uint8>;
+				Y0 = Private::ScalePixel(Y0,
+										 FRange::Inclusive(16, 235),
+										 FRange::Inclusive(MIN_uint8, MAX_uint8));
+				Y1 = Private::ScalePixel(Y1,
+										 FRange::Inclusive(16, 235),
+										 FRange::Inclusive(MIN_uint8, MAX_uint8));
+
+			}
+
+			NewBuffer[Y * MonoStride + X] = 
+				Private::CreateColor(Y0, Y0, Y0, InScaleRange ? 235 : 255);
+			NewBuffer[Y * MonoStride + X + 1] = 
+				Private::CreateColor(Y1, Y1, Y1, InScaleRange ? 235 : 255);
+		}
+	});
+
+	return NewBuffer;
+}
+
+TArray<FColor> UEConvertI420ToBGRA(const UE::CaptureManager::FMediaTextureSample* InSample)
+{
+	using namespace UE::CaptureManager;
+
+	EMediaTexturePixelFormat OutputFormat = EMediaTexturePixelFormat::U8_BGRA;
+
+	TArray<FColor> NewBuffer;
+	NewBuffer.SetNumUninitialized(InSample->Stride * InSample->Dimensions.Y);
+
+	uint32 StrideY = InSample->Stride;
+	uint32 StrideUV = InSample->Stride / 2;
+	int32 BGRAStride = InSample->Stride;
+
+	uint32 YPlaneSize = StrideY * InSample->Dimensions.Y;
+	uint32 UPlaneSize = StrideUV * (InSample->Dimensions.Y / 2);
+	uint32 VPlaneSize = InSample->Buffer.Num() - (YPlaneSize + UPlaneSize);
+
+	const TArrayView<const uint8> YPlane = MakeArrayView(InSample->Buffer.GetData(), YPlaneSize);
+	const TArrayView<const uint8> UPlane = MakeArrayView(InSample->Buffer.GetData() + YPlaneSize, UPlaneSize);
+	const TArrayView<const uint8> VPlane = MakeArrayView(InSample->Buffer.GetData() + YPlaneSize + UPlaneSize, VPlaneSize);
+
+	ParallelFor(InSample->Dimensions.Y, [&NewBuffer, YPlane, UPlane, VPlane, StrideY, StrideUV, BGRAStride, InSample](int32 Y) mutable
+	{
+		for (int32 X = 0; X < InSample->Dimensions.X; ++X)
+		{
+			int32 YIndex = Y * StrideY + X;
+			int32 UVIndex = (Y / 2) * StrideUV + (X / 2);
+
+			uint8_t Y0 = YPlane[YIndex];
+			uint8_t U = UPlane[UVIndex];
+			uint8_t V = VPlane[UVIndex];
+
+			uint8 R, G, B;
+			Private::RGBfromYUV(R, G, B, Y0, U, V);
+
+			NewBuffer[Y * BGRAStride + X] = Private::CreateColor(R, G, B, 255);
+		}
+	});
+
+	return NewBuffer;
+}
+
+TArray<FColor> UEConvertNV12ToBGRA(const UE::CaptureManager::FMediaTextureSample* InSample)
+{
+	using namespace UE::CaptureManager;
+
+	EMediaTexturePixelFormat OutputFormat = EMediaTexturePixelFormat::U8_BGRA;
+
+	TArray<FColor> NewBuffer;
+	NewBuffer.SetNumUninitialized(InSample->Stride * InSample->Dimensions.Y);
+
+	uint32 StrideY = InSample->Stride;
+	int32 BGRAStride = InSample->Stride;
+
+	uint32 YPlaneSize = StrideY * InSample->Dimensions.Y;
+	uint32 UVPlaneSize = InSample->Buffer.Num() - YPlaneSize;
+
+	const TArrayView<const uint8> YPlane = MakeArrayView(InSample->Buffer.GetData(), YPlaneSize);
+	const TArrayView<const uint8> UVPlane = MakeArrayView(InSample->Buffer.GetData() + YPlaneSize, UVPlaneSize);
+
+	ParallelFor(InSample->Dimensions.Y, [&NewBuffer, YPlane, UVPlane, StrideY, BGRAStride, InSample](int32 Y) mutable
+	{
+		for (int32 X = 0; X < InSample->Dimensions.X; ++X)
+		{
+			int32 YIndex = Y * StrideY + X;
+			int32 UVIndex = (Y / 2) * StrideY + (X & ~1);
+
+			uint8_t Y0 = YPlane[YIndex];
+			uint8_t U = UVPlane[UVIndex];
+			uint8_t V = UVPlane[UVIndex + 1];
+
+			uint8 R, G, B;
+			Private::RGBfromYUV(R, G, B, Y0, U, V);
+
+			NewBuffer[Y * BGRAStride + X] = Private::CreateColor(R, G, B, 255);
+		}
+	});
+
+	return NewBuffer;
+}
+
+TArray<FColor> UEConvertYUY2ToBGRA(const UE::CaptureManager::FMediaTextureSample* InSample)
+{
+	using namespace UE::CaptureManager;
+
+	EMediaTexturePixelFormat OutputFormat = EMediaTexturePixelFormat::U8_BGRA;
+
+	TArray<FColor> NewBuffer;
+	NewBuffer.SetNumUninitialized(InSample->Stride * InSample->Dimensions.Y);
+
+	int32 YUY2Stride = InSample->Stride * 2;
+	int32 BGRAStride = InSample->Stride;
+
+	ParallelFor(InSample->Dimensions.Y, [&NewBuffer, YUY2Stride, BGRAStride, InSample](int32 Y) mutable
+	{
+		for (int32 X = 0; X < InSample->Dimensions.X; X += 2)
+		{
+			int32 YIndex = Y * YUY2Stride + X * 2;
+
+			uint8 Y0 = InSample->Buffer[YIndex];
+			uint8 U = InSample->Buffer[YIndex + 1];
+			uint8 Y1 = InSample->Buffer[YIndex + 2];
+			uint8 V = InSample->Buffer[YIndex + 3];
+
+			uint8 R0, G0, B0;
+			Private::RGBfromYUV(R0, G0, B0, Y0, U, V);
+
+			uint8 R1, G1, B1;
+			Private::RGBfromYUV(R1, G1, B1, Y1, U, V);
+
+			NewBuffer[Y * BGRAStride + X] = Private::CreateColor(R0, G0, B0, 255);
+			NewBuffer[Y * BGRAStride + (X + 1)] = Private::CreateColor(R1, G1, B1, 255);
+		}
+	});
+
+	return NewBuffer;
+}
+
+}
\ No newline at end of file
diff --git a/Engine/Plugins/VirtualProduction/CaptureManager/CaptureManagerApp/Source/CaptureManagerMediaRW/Private/Writers/WindowsImageWriter.cpp b/Engine/Plugins/VirtualProduction/CaptureManager/CaptureManagerApp/Source/CaptureManagerMediaRW/Private/Writers/WindowsImageWriter.cpp
--- a/Engine/Plugins/VirtualProduction/CaptureManager/CaptureManagerApp/Source/CaptureManagerMediaRW/Private/Writers/WindowsImageWriter.cpp
+++ b/Engine/Plugins/VirtualProduction/CaptureManager/CaptureManagerApp/Source/CaptureManagerMediaRW/Private/Writers/WindowsImageWriter.cpp
@@ -6,7 +6,7 @@
 
 #include "IImageWrapperModule.h"
 #include "Modules/ModuleManager.h"
-#include "Async/ParallelFor.h"
+#include "Utils/MediaPixelFormatConversions.h"
 
 #if PLATFORM_WINDOWS && !UE_SERVER
 
@@ -122,29 +122,6 @@ FString GetPixelFormatString(EMediaTexturePixelFormat InPixelFormat)
 	}
 }
 
-template<typename T>
-T ScalePixel(T InPixel, TRange<T> InInputRange, TRange<T> InOutputRange)
-{
-	const float Factor = static_cast<float>(InOutputRange.GetUpperBoundValue()) / (InInputRange.GetUpperBoundValue() - InInputRange.GetLowerBoundValue());
-	const uint64 NewPixelValue = Factor * (InPixel - InInputRange.GetLowerBoundValue()) + 0.5;
-	return FMath::Clamp(NewPixelValue, InOutputRange.GetLowerBoundValue(), InOutputRange.GetUpperBoundValue());
-}
-
-static inline void RGBfromYUV(uint8& R, uint8& G, uint8& B, uint8 Y, uint8 U, uint8 V)
-{
-	int64 BigY = Y - 16;
-	int64 BigU = U - 128;
-	int64 BigV = V - 128;
-
-	int64 BigR = (298 * BigY + 409 * BigV + 128) >> 8;
-	int64 BigG = (298 * BigY - 100 * BigU - 208 * BigV + 128) >> 8;
-	int64 BigB = (298 * BigY + 516 * BigU + 128) >> 8;
-
-	R = FMath::Clamp(BigR, 0, 255);
-	G = FMath::Clamp(BigG, 0, 255);
-	B = FMath::Clamp(BigB, 0, 255);
-}
-
 class FConverter
 {
 public:
@@ -271,13 +248,13 @@ TOptional<FText> FWindowsImageWriter::Append(UE::CaptureManager::FMediaTextureSa
 		InSample->DesiredFormat = EMediaTexturePixelFormat::U8_Mono;
 		if (InSample->CurrentFormat == UE::CaptureManager::EMediaTexturePixelFormat::U8_YUY2)
 		{
-			InSample->Buffer = ConvertYUY2ToMono(InSample, false);
+			InSample->Buffer = ConvertYUY2ToMono(InSample, true);
 			InSample->CurrentFormat = EMediaTexturePixelFormat::U8_Mono;
 		}
 		else if (InSample->CurrentFormat == UE::CaptureManager::EMediaTexturePixelFormat::U8_I420 ||
 				 InSample->CurrentFormat == UE::CaptureManager::EMediaTexturePixelFormat::U8_NV12)
 		{
-			InSample->Buffer = ConvertYUVToMono(InSample);
+			InSample->Buffer = ConvertYUVToMono(InSample, true);
 			InSample->CurrentFormat = EMediaTexturePixelFormat::U8_Mono;
 		}
 		else
@@ -430,212 +407,6 @@ GUID FWindowsImageWriter::GetEncoderGuidBasedOnFormat(const FString& InFormat)
 	return UnsupportedEncoder;
 }
 
-TArray<uint8> FWindowsImageWriter::ConvertYUVToMono(const UE::CaptureManager::FMediaTextureSample* InSample, bool InScaleRange)
-{
-	using namespace UE::CaptureManager;
-
-	TArray<uint8> NewBuffer;
-	NewBuffer.SetNumUninitialized(InSample->Stride * InSample->Dimensions.Y);
-
-	uint32 StrideY = InSample->Stride;
-	int32 MonoStride = InSample->Stride;
-
-	ParallelFor(InSample->Dimensions.Y, [&NewBuffer, StrideY, MonoStride, InScaleRange, InSample](int32 Y) mutable
-	{
-		for (int32 X = 0; X < InSample->Dimensions.X; ++X)
-		{
-			int32 YIndex = Y * StrideY + X;
-			uint8 PixelValue = InSample->Buffer[YIndex];
-
-			if (InScaleRange)
-			{
-				using FRange = TRange<uint8>;
-				PixelValue = Private::ScalePixel(PixelValue,
-												 FRange::Inclusive(16, 235),
-												 FRange::Inclusive(MIN_uint8, MAX_uint8));
-			}
-
-			NewBuffer[Y * MonoStride + X] = PixelValue;
-		}
-	});
-
-	return NewBuffer;
-}
-
-TArray<uint8> FWindowsImageWriter::ConvertYUY2ToMono(const UE::CaptureManager::FMediaTextureSample* InSample, bool InScaleRange)
-{
-	using namespace UE::CaptureManager;
-
-	TArray<uint8> NewBuffer;
-	NewBuffer.SetNumUninitialized(InSample->Stride * InSample->Dimensions.Y);
-
-	int32 YUY2Stride = InSample->Stride * 2;
-	int32 MonoStride = InSample->Stride;
-
-	ParallelFor(InSample->Dimensions.Y, [&NewBuffer, YUY2Stride, MonoStride, InScaleRange, InSample](int32 Y) mutable
-	{
-		for (int32 X = 0; X < InSample->Dimensions.X; X += 2)
-		{
-			int32 YIndex = Y * YUY2Stride + X * 2;
-
-			uint8_t Y0 = InSample->Buffer[YIndex];
-			uint8_t Y1 = InSample->Buffer[YIndex + 2];
-
-			if (InScaleRange)
-			{
-				using FRange = TRange<uint8>;
-				Y0 = Private::ScalePixel(Y0,
-										 FRange::Inclusive(16, 235),
-										 FRange::Inclusive(MIN_uint8, MAX_uint8));
-				Y1 = Private::ScalePixel(Y1,
-										 FRange::Inclusive(16, 235),
-										 FRange::Inclusive(MIN_uint8, MAX_uint8));
-
-			}
-
-			NewBuffer[Y * MonoStride + X] = Y0;
-			NewBuffer[Y * MonoStride + X + 1] = Y1;
-		}
-	});
-
-	return NewBuffer;
-}
-
-TArray<uint8> FWindowsImageWriter::ConvertI420ToBGRA(const UE::CaptureManager::FMediaTextureSample* InSample)
-{
-	using namespace UE::CaptureManager;
-
-	EMediaTexturePixelFormat OutputFormat = EMediaTexturePixelFormat::U8_BGRA;
-	uint32 Channels = GetNumberOfChannels(OutputFormat);
-
-	TArray<uint8> NewBuffer;
-	NewBuffer.SetNumUninitialized(InSample->Stride * InSample->Dimensions.Y * Channels);
-
-	uint32 StrideY = InSample->Stride;
-	uint32 StrideUV = InSample->Stride / 2;
-	int32 BGRAStride = InSample->Stride * Channels;
-
-	uint32 YPlaneSize = StrideY * InSample->Dimensions.Y;
-	uint32 UPlaneSize = StrideUV * (InSample->Dimensions.Y / 2);
-	uint32 VPlaneSize = InSample->Buffer.Num() - (YPlaneSize + UPlaneSize);
-
-	const TArrayView<const uint8> YPlane = MakeArrayView(InSample->Buffer.GetData(), YPlaneSize);
-	const TArrayView<const uint8> UPlane = MakeArrayView(InSample->Buffer.GetData() + YPlaneSize, UPlaneSize);
-	const TArrayView<const uint8> VPlane = MakeArrayView(InSample->Buffer.GetData() + YPlaneSize + UPlaneSize, VPlaneSize);
-
-	ParallelFor(InSample->Dimensions.Y, [&NewBuffer, YPlane, UPlane, VPlane, StrideY, StrideUV, BGRAStride, InSample](int32 Y) mutable
-	{
-		for (int32 X = 0; X < InSample->Dimensions.X; ++X)
-		{
-			int32 YIndex = Y * StrideY + X;
-			int32 UVIndex = (Y / 2) * StrideUV + (X / 2);
-
-			uint8_t Y0 = YPlane[YIndex];
-			uint8_t U = UPlane[UVIndex];
-			uint8_t V = VPlane[UVIndex];
-
-			uint8 R, G, B;
-			Private::RGBfromYUV(R, G, B, Y0, U, V);
-
-			NewBuffer[Y * BGRAStride + X * 4] = B;
-			NewBuffer[Y * BGRAStride + X * 4 + 1] = G;
-			NewBuffer[Y * BGRAStride + X * 4 + 2] = R;
-			NewBuffer[Y * BGRAStride + X * 4 + 3] = 255;
-		}
-	});
-
-	return NewBuffer;
-}
-
-TArray<uint8> FWindowsImageWriter::ConvertNV12ToBGRA(const UE::CaptureManager::FMediaTextureSample* InSample)
-{
-	using namespace UE::CaptureManager;
-
-	EMediaTexturePixelFormat OutputFormat = EMediaTexturePixelFormat::U8_BGRA;
-	uint32 Channels = GetNumberOfChannels(OutputFormat);
-
-	TArray<uint8> NewBuffer;
-	NewBuffer.SetNumUninitialized(InSample->Stride * InSample->Dimensions.Y * Channels);
-
-	uint32 StrideY = InSample->Stride;
-	int32 BGRAStride = InSample->Stride * Channels;
-
-	uint32 YPlaneSize = StrideY * InSample->Dimensions.Y;
-	uint32 UVPlaneSize = InSample->Buffer.Num() - YPlaneSize;
-
-	const TArrayView<const uint8> YPlane = MakeArrayView(InSample->Buffer.GetData(), YPlaneSize);
-	const TArrayView<const uint8> UVPlane = MakeArrayView(InSample->Buffer.GetData() + YPlaneSize, UVPlaneSize);
-
-	ParallelFor(InSample->Dimensions.Y, [&NewBuffer, YPlane, UVPlane, StrideY, BGRAStride, InSample, Channels](int32 Y) mutable
-	{
-		for (int32 X = 0; X < InSample->Dimensions.X; ++X)
-		{
-			int32 YIndex = Y * StrideY + X;
-			int32 UVIndex = (Y / 2) * StrideY + (X & ~1);
-
-			uint8_t Y0 = YPlane[YIndex];
-			uint8_t U = UVPlane[UVIndex];
-			uint8_t V = UVPlane[UVIndex + 1];
-
-			uint8 R, G, B;
-			Private::RGBfromYUV(R, G, B, Y0, U, V);
-
-			NewBuffer[Y * BGRAStride + X * 4] = B;
-			NewBuffer[Y * BGRAStride + X * 4 + 1] = G;
-			NewBuffer[Y * BGRAStride + X * 4 + 2] = R;
-			NewBuffer[Y * BGRAStride + X * 4 + 3] = 255;
-		}	
-	});
-
-	return NewBuffer;
-}
-
-TArray<uint8> FWindowsImageWriter::ConvertYUY2ToBGRA(const UE::CaptureManager::FMediaTextureSample* InSample)
-{
-	using namespace UE::CaptureManager;
-
-	EMediaTexturePixelFormat OutputFormat = EMediaTexturePixelFormat::U8_BGRA;
-
-	uint32 Channels = GetNumberOfChannels(OutputFormat);
-
-	TArray<uint8> NewBuffer;
-	NewBuffer.SetNumUninitialized(InSample->Stride * InSample->Dimensions.Y * Channels);
-
-	int32 YUY2Stride = InSample->Stride * 2;
-	int32 BGRAStride = InSample->Stride * Channels;
-
-	ParallelFor(InSample->Dimensions.Y, [&NewBuffer, YUY2Stride, BGRAStride, InSample](int32 Y) mutable
-	{
-		for (int32 X = 0; X < InSample->Dimensions.X; X += 2)
-		{
-			int32 YIndex = Y * YUY2Stride + X * 2;
-
-			uint8 Y0 = InSample->Buffer[YIndex];
-			uint8 U = InSample->Buffer[YIndex + 1];
-			uint8 Y1 = InSample->Buffer[YIndex + 2];
-			uint8 V = InSample->Buffer[YIndex + 3];
-
-			uint8 R0, G0, B0;
-			Private::RGBfromYUV(R0, G0, B0, Y0, U, V);
-
-			uint8 R1, G1, B1;
-			Private::RGBfromYUV(R1, G1, B1, Y1, U, V);
-
-			NewBuffer[Y * BGRAStride + X * 4] = B0;
-			NewBuffer[Y * BGRAStride + X * 4 + 1] = G0;
-			NewBuffer[Y * BGRAStride + X * 4 + 2] = R0;
-			NewBuffer[Y * BGRAStride + X * 4 + 3] = 255;
-
-			NewBuffer[Y * BGRAStride + (X + 1) * 4] = B1;
-			NewBuffer[Y * BGRAStride + (X + 1) * 4 + 1] = G1;
-			NewBuffer[Y * BGRAStride + (X + 1) * 4 + 2] = R1;
-			NewBuffer[Y * BGRAStride + (X + 1) * 4 + 3] = 255;
-		}
-	});
-
-	return NewBuffer;
-}
-
 FString FWindowsImageWriter::CreateFrameFileName()
 {
 	FString FrameIndexStr = FString::Printf(TEXT("%06d"), FrameNumber);
diff --git a/Engine/Plugins/VirtualProduction/CaptureManager/CaptureManagerApp/Source/CaptureManagerMediaRW/Private/Writers/WindowsImageWriter.h b/Engine/Plugins/VirtualProduction/CaptureManager/CaptureManagerApp/Source/CaptureManagerMediaRW/Private/Writers/WindowsImageWriter.h
--- a/Engine/Plugins/VirtualProduction/CaptureManager/CaptureManagerApp/Source/CaptureManagerMediaRW/Private/Writers/WindowsImageWriter.h
+++ b/Engine/Plugins/VirtualProduction/CaptureManager/CaptureManagerApp/Source/CaptureManagerMediaRW/Private/Writers/WindowsImageWriter.h
@@ -47,13 +47,6 @@ public:
 private:
 
 	GUID GetEncoderGuidBasedOnFormat(const FString& InFormat);
-
-	// Scaling from 16-235 range to 0-255 range
-	TArray<uint8> ConvertYUVToMono(const UE::CaptureManager::FMediaTextureSample* InSample, bool InScaleRange = true); // Any YUV to Mono
-	TArray<uint8> ConvertYUY2ToMono(const UE::CaptureManager::FMediaTextureSample* InSample, bool InScaleRange = true);
-	TArray<uint8> ConvertI420ToBGRA(const UE::CaptureManager::FMediaTextureSample* InSample);
-	TArray<uint8> ConvertNV12ToBGRA(const UE::CaptureManager::FMediaTextureSample* InSample);
-	TArray<uint8> ConvertYUY2ToBGRA(const UE::CaptureManager::FMediaTextureSample* InSample);
 	FString CreateFrameFileName();
 
 	TComPtr<IWICImagingFactory> WindowsImagingFactory;
diff --git a/Engine/Plugins/VirtualProduction/CaptureManager/CaptureManagerApp/Source/CaptureManagerMediaRW/Public/Utils/MediaPixelFormatConversions.h b/Engine/Plugins/VirtualProduction/CaptureManager/CaptureManagerApp/Source/CaptureManagerMediaRW/Public/Utils/MediaPixelFormatConversions.h
new file mode 100644
--- /dev/null
+++ b/Engine/Plugins/VirtualProduction/CaptureManager/CaptureManagerApp/Source/CaptureManagerMediaRW/Public/Utils/MediaPixelFormatConversions.h
@@ -0,0 +1,20 @@
+// Copyright Epic Games, Inc. All Rights Reserved.
+
+#include "MediaSample.h"
+
+namespace UE::CaptureManager
+{
+
+CAPTUREMANAGERMEDIARW_API TArray<uint8> ConvertYUVToMono(const UE::CaptureManager::FMediaTextureSample* InSample, bool InScaleRange); // Any YUV to Mono
+CAPTUREMANAGERMEDIARW_API TArray<uint8> ConvertYUY2ToMono(const UE::CaptureManager::FMediaTextureSample* InSample, bool InScaleRange);
+CAPTUREMANAGERMEDIARW_API TArray<uint8> ConvertI420ToBGRA(const UE::CaptureManager::FMediaTextureSample* InSample);
+CAPTUREMANAGERMEDIARW_API TArray<uint8> ConvertNV12ToBGRA(const UE::CaptureManager::FMediaTextureSample* InSample);
+CAPTUREMANAGERMEDIARW_API TArray<uint8> ConvertYUY2ToBGRA(const UE::CaptureManager::FMediaTextureSample* InSample);
+
+CAPTUREMANAGERMEDIARW_API TArray<FColor> UEConvertYUVToMono(const UE::CaptureManager::FMediaTextureSample* InSample, bool InScaleRange); // Any YUV to Mono
+CAPTUREMANAGERMEDIARW_API TArray<FColor> UEConvertYUY2ToMono(const UE::CaptureManager::FMediaTextureSample* InSample, bool InScaleRange);
+CAPTUREMANAGERMEDIARW_API TArray<FColor> UEConvertI420ToBGRA(const UE::CaptureManager::FMediaTextureSample* InSample);
+CAPTUREMANAGERMEDIARW_API TArray<FColor> UEConvertNV12ToBGRA(const UE::CaptureManager::FMediaTextureSample* InSample);
+CAPTUREMANAGERMEDIARW_API TArray<FColor> UEConvertYUY2ToBGRA(const UE::CaptureManager::FMediaTextureSample* InSample);
+
+}
\ No newline at end of file
diff --git a/Engine/Plugins/VirtualProduction/CaptureManager/CaptureManagerCore/Source/CaptureManagerTakeMetadata/CaptureManagerTakeMetadata.Build.cs b/Engine/Plugins/VirtualProduction/CaptureManager/CaptureManagerCore/Source/CaptureManagerTakeMetadata/CaptureManagerTakeMetadata.Build.cs
--- a/Engine/Plugins/VirtualProduction/CaptureManager/CaptureManagerCore/Source/CaptureManagerTakeMetadata/CaptureManagerTakeMetadata.Build.cs
+++ b/Engine/Plugins/VirtualProduction/CaptureManager/CaptureManagerCore/Source/CaptureManagerTakeMetadata/CaptureManagerTakeMetadata.Build.cs
@@ -6,6 +6,11 @@ public class CaptureManagerTakeMetadata : ModuleRules
 {
 	public CaptureManagerTakeMetadata(ReadOnlyTargetRules Target) : base(Target)
 	{
+		PublicDependencyModuleNames.AddRange(new string[]
+		{
+			"ImageCore"
+		});
+
 		PrivateDependencyModuleNames.AddRange(new string[] {
 			"Core",
 			"Projects",
diff --git a/Engine/Plugins/VirtualProduction/CaptureManager/CaptureManagerCore/Source/CaptureManagerTakeMetadata/Private/CaptureManagerTakeMetadata.cpp b/Engine/Plugins/VirtualProduction/CaptureManager/CaptureManagerCore/Source/CaptureManagerTakeMetadata/Private/CaptureManagerTakeMetadata.cpp
--- a/Engine/Plugins/VirtualProduction/CaptureManager/CaptureManagerCore/Source/CaptureManagerTakeMetadata/Private/CaptureManagerTakeMetadata.cpp
+++ b/Engine/Plugins/VirtualProduction/CaptureManager/CaptureManagerCore/Source/CaptureManagerTakeMetadata/Private/CaptureManagerTakeMetadata.cpp
@@ -28,6 +28,26 @@ FTakeThumbnailData::FTakeThumbnailData(TArray<uint8> InThumbnailData)
 {
 }
 
+FTakeThumbnailData::FTakeThumbnailData(TArray<FColor> InDecompressedImageData, uint32 InWidth, uint32 InHeight, ERawImageFormat::Type InFormat)
+	: Thumbnail(TInPlaceType<FRawImage>(), MoveTemp(InDecompressedImageData), InWidth, InHeight, InFormat)
+{
+}
+
+void FTakeThumbnailData::operator=(FString InImagePath)
+{
+	Thumbnail.Set<FString>(MoveTemp(InImagePath));
+}
+
+void FTakeThumbnailData::operator=(TArray<uint8> InCompressedImageData)
+{
+	Thumbnail.Set<TArray<uint8>>(MoveTemp(InCompressedImageData));
+}
+
+void FTakeThumbnailData::operator=(FRawImage InRawImage)
+{
+	Thumbnail.Set<FRawImage>(MoveTemp(InRawImage));
+}
+
 TOptional<TArray<uint8>> FTakeThumbnailData::GetThumbnailData() const
 {
 	if (Thumbnail.IsType<FString>())
@@ -58,6 +78,16 @@ TOptional<FString> FTakeThumbnailData::GetThumbnailPath() const
 	return {};
 }
 
+TOptional<FTakeThumbnailData::FRawImage> FTakeThumbnailData::GetRawImage() const
+{
+	if (Thumbnail.IsType<FRawImage>())
+	{
+		return Thumbnail.Get<FRawImage>();
+	}
+
+	return {};
+}
+
 const FString FTakeMetadata::FileExtension = TEXT("cptake");
 
 static FTakeMetadata::FSchemaVersion ParseSchemaVersion(const rapidjson::Value& InVersionObject)
diff --git a/Engine/Plugins/VirtualProduction/CaptureManager/CaptureManagerCore/Source/CaptureManagerTakeMetadata/Public/CaptureManagerTakeMetadata.h b/Engine/Plugins/VirtualProduction/CaptureManager/CaptureManagerCore/Source/CaptureManagerTakeMetadata/Public/CaptureManagerTakeMetadata.h
--- a/Engine/Plugins/VirtualProduction/CaptureManager/CaptureManagerCore/Source/CaptureManagerTakeMetadata/Public/CaptureManagerTakeMetadata.h
+++ b/Engine/Plugins/VirtualProduction/CaptureManager/CaptureManagerCore/Source/CaptureManagerTakeMetadata/Public/CaptureManagerTakeMetadata.h
@@ -9,18 +9,42 @@
 #include "Templates/ValueOrError.h"
 #include "Internationalization/Text.h"
 
+#include "Math/MathFwd.h"
+#include "ImageCore.h"
+
 /** Class for constructing take thumbnails and accessing thumbnail data. */
 class CAPTUREMANAGERTAKEMETADATA_API FTakeThumbnailData
 {
 public:
+
+	struct FRawImage
+	{
+		TArray<FColor> DecompressedImageData;
+		uint32 Width;
+		uint32 Height;
+		ERawImageFormat::Type Format;
+	};
+
 	/** Default thumbnail constructor. */
 	FTakeThumbnailData();
 
 	/** Construct thumbnail from file path. */
-	FTakeThumbnailData(FString InThumbnailPath);
+	FTakeThumbnailData(FString InImagePath);
 
 	/** Construct thumbnail from file data array. */
-	FTakeThumbnailData(TArray<uint8> InThumbnailData);
+	FTakeThumbnailData(TArray<uint8> InCompressedImageData);
+
+	/** Construct thumbnail from raw data array. */
+	FTakeThumbnailData(TArray<FColor> InDecompressedImageData, uint32 InWidth, uint32 InHeight, ERawImageFormat::Type InFormat);
+
+	/** Construct thumbnail from file path. */
+	void operator=(FString InImagePath);
+
+	/** Construct thumbnail from file data array. */
+	void operator=(TArray<uint8> InCompressedImageData);
+
+	/** Assign thumbnail from raw data array. */
+	void operator=(FRawImage InRawImage);
 
 	/** Get thumbnail. */
 	TOptional<TArray<uint8>> GetThumbnailData() const;
@@ -28,8 +52,12 @@ public:
 	/** Get thumbnail file path. */
 	TOptional<FString> GetThumbnailPath() const;
 
+	/** Get decompressed image. */
+	TOptional<FRawImage> GetRawImage() const;
+
 private:
-	using FThumbnail = TVariant<FEmptyVariantState, FString, TArray<uint8>>;
+
+	using FThumbnail = TVariant<FEmptyVariantState, FString, TArray<uint8>, FRawImage>;
 	FThumbnail Thumbnail;
 };
 
diff --git a/Engine/Plugins/VirtualProduction/CaptureManager/CaptureManagerDevices/Source/MonoVideoLiveLinkDevice/Private/MonoVideoLiveLinkDevice.cpp b/Engine/Plugins/VirtualProduction/CaptureManager/CaptureManagerDevices/Source/MonoVideoLiveLinkDevice/Private/MonoVideoLiveLinkDevice.cpp
--- a/Engine/Plugins/VirtualProduction/CaptureManager/CaptureManagerDevices/Source/MonoVideoLiveLinkDevice/Private/MonoVideoLiveLinkDevice.cpp
+++ b/Engine/Plugins/VirtualProduction/CaptureManager/CaptureManagerDevices/Source/MonoVideoLiveLinkDevice/Private/MonoVideoLiveLinkDevice.cpp
@@ -39,6 +39,7 @@ FMediaRWManager& GetMediaRWManager()
 {
 	return FModuleManager::LoadModuleChecked<FCaptureManagerMediaRWModule>("CaptureManagerMediaRW").Get();
 }
+
 }
 
 const UMonoVideoLiveLinkDeviceSettings* UMonoVideoLiveLinkDevice::GetSettings() const
@@ -123,20 +124,15 @@ void UMonoVideoLiveLinkDevice::UpdateTakeList_Implementation(UIngestCapability_U
 		RemoveAllTakes();
 
 		FString StoragePath = GetSettings()->TakeDirectory.Path;
-		int32 PathsInterrogatedCount = 0;
-		const int32 PathsToInterrogateInOneRun = 100000;
+		int32 DirectoriesInterrogatedCount = 0;
+		const int32 DirectoriesToInterrogateInOneRun = 200;
 	
 		TArray<FString> SupportedVideoFiles;
 		const bool bIterationResult = IFileManager::Get().IterateDirectoryRecursively(*StoragePath, [&](const TCHAR* InFileNameOrDirectory, bool bInIsDirectory)
 		{
-			if (++PathsInterrogatedCount > PathsToInterrogateInOneRun)
+			if (bInIsDirectory)
 			{
-				EAppReturnType::Type Result = FMessageDialog::Open(EAppMsgType::OkCancel, FText::Format(LOCTEXT("MonoVideoDevice_ProceedWithTakesDiscovery", "Analyzed {0} paths in search for takes. Do you want to analyze {1} more?"), PathsToInterrogateInOneRun, PathsToInterrogateInOneRun));
-				if (Result == EAppReturnType::Ok)
-				{
-					PathsInterrogatedCount = 0;
-				}
-				else 
+				if (++DirectoriesInterrogatedCount > DirectoriesToInterrogateInOneRun)
 				{
 					return false;
 				}
@@ -292,10 +288,13 @@ TOptional<FTakeMetadata> UMonoVideoLiveLinkDevice::ReadTake(const FString& InCur
 
 			if (ThumbnailRawColorData.Num() > 0)
 			{
-				TArray<uint8> ThumbnailData;
-				FImageUtils::ThumbnailCompressImageArray(Sample->Dimensions.X, Sample->Dimensions.Y, ThumbnailRawColorData, ThumbnailData);
+				FTakeThumbnailData::FRawImage RawImage;
+				RawImage.DecompressedImageData = MoveTemp(ThumbnailRawColorData);
+				RawImage.Width = Sample->Dimensions.X;
+				RawImage.Height = Sample->Dimensions.Y;
+				RawImage.Format = ERawImageFormat::BGRA8;
 
-				TakeMetadata.Thumbnail = ThumbnailData;
+				TakeMetadata.Thumbnail = MoveTemp(RawImage);
 			}
 		}
 	}
diff --git a/Engine/Plugins/VirtualProduction/CaptureManager/CaptureManagerDevices/Source/StereoVideoLiveLinkDevice/Private/StereoVideoLiveLinkDevice.cpp b/Engine/Plugins/VirtualProduction/CaptureManager/CaptureManagerDevices/Source/StereoVideoLiveLinkDevice/Private/StereoVideoLiveLinkDevice.cpp
--- a/Engine/Plugins/VirtualProduction/CaptureManager/CaptureManagerDevices/Source/StereoVideoLiveLinkDevice/Private/StereoVideoLiveLinkDevice.cpp
+++ b/Engine/Plugins/VirtualProduction/CaptureManager/CaptureManagerDevices/Source/StereoVideoLiveLinkDevice/Private/StereoVideoLiveLinkDevice.cpp
@@ -190,25 +190,19 @@ TMap<FString, UStereoVideoLiveLinkDevice::FTakeWithComponents> UStereoVideoLiveL
 {
 	TMap<FString, FTakeWithComponents> TakesWithComponentsCandidates;
 
-	int32 PathsInterrogatedCount = 0;
-	const int32 PathsToInterrogateInOneRun = 100000;
+	int32 DirectoriesInterrogatedCount = 0;
+	const int32 DirectoriesToInterrogateInOneRun = 200;
 
 	const bool bIterationResult = IFileManager::Get().IterateDirectoryRecursively(*StoragePath, [&](const TCHAR* InPath, bool bInIsDirectory)
 		{
-			if (++PathsInterrogatedCount > PathsToInterrogateInOneRun)
+			if (bInIsDirectory)
 			{
-				EAppReturnType::Type Result = FMessageDialog::Open(EAppMsgType::OkCancel, FText::Format(LOCTEXT("StereoVideoDevice_ProceedWithTakesDiscovery", "Analyzed {0} paths in search for takes. Do you want to analyze {1} more?"), PathsToInterrogateInOneRun, PathsToInterrogateInOneRun));
-				if (Result == EAppReturnType::Ok)
-				{
-					PathsInterrogatedCount = 0;
-				}
-				else
+				if (++DirectoriesInterrogatedCount > DirectoriesToInterrogateInOneRun)
 				{
 					return false;
 				}
 			}
 
-
 			FString Extension = FPaths::GetExtension(InPath);
 			if (UE::StereoVideoLiveLinkDevice::Private::SupportedVideoExtensions.Contains(Extension))
 			{
@@ -508,11 +502,15 @@ TOptional<FTakeMetadata> UStereoVideoLiveLinkDevice::CreateTakeMetadata(const FT
 
 						TArray<FColor> ThumbnailRawColorData = ThumbnailExtractor.ExtractThumbnail(Sample);
 
-						if (ThumbnailRawColorData.Num() > 0) {
-							TArray<uint8> ThumbnailData;
-							FImageUtils::ThumbnailCompressImageArray(Sample->Dimensions.X, Sample->Dimensions.Y, ThumbnailRawColorData, ThumbnailData);
+						if (ThumbnailRawColorData.Num() > 0) 
+						{
+							FTakeThumbnailData::FRawImage RawImage;
+							RawImage.DecompressedImageData = MoveTemp(ThumbnailRawColorData);
+							RawImage.Width = Sample->Dimensions.X;
+							RawImage.Height = Sample->Dimensions.Y;
+							RawImage.Format = ERawImageFormat::BGRA8;
 
-							TakeMetadata.Thumbnail = ThumbnailData;
+							TakeMetadata.Thumbnail = MoveTemp(RawImage);
 						}
 					}
 				}
@@ -552,11 +550,12 @@ TOptional<FTakeMetadata> UStereoVideoLiveLinkDevice::CreateTakeMetadata(const FT
 
 							Video.Path = TakeComponent.Path;
 							{
-								TObjectPtr<class UTexture2D> Texture = FImageUtils::ImportFileAsTexture2D(InPath);
-								if (Texture)
+								FImage Image;
+								FImageUtils::LoadImage(InPath, Image);
+								if (Image.IsImageInfoValid())
 								{
-									Video.FrameWidth = Texture->GetSizeX();
-									Video.FrameHeight = Texture->GetSizeY();
+									Video.FrameWidth = Image.SizeX;
+									Video.FrameHeight = Image.SizeY;
 								}
 							}
 						}
diff --git a/Engine/Plugins/VirtualProduction/CaptureManager/CaptureManagerDevices/Source/VideoLiveLinkDeviceCommon/Private/Utils/VideoDeviceThumbnailExtractor.cpp b/Engine/Plugins/VirtualProduction/CaptureManager/CaptureManagerDevices/Source/VideoLiveLinkDeviceCommon/Private/Utils/VideoDeviceThumbnailExtractor.cpp
--- a/Engine/Plugins/VirtualProduction/CaptureManager/CaptureManagerDevices/Source/VideoLiveLinkDeviceCommon/Private/Utils/VideoDeviceThumbnailExtractor.cpp
+++ b/Engine/Plugins/VirtualProduction/CaptureManager/CaptureManagerDevices/Source/VideoLiveLinkDeviceCommon/Private/Utils/VideoDeviceThumbnailExtractor.cpp
@@ -7,6 +7,8 @@
 #include "CaptureManagerMediaRWModule.h"
 #include "VideoLiveLinkDeviceLog.h"
 
+#include "Utils/MediaPixelFormatConversions.h"
+
 DEFINE_LOG_CATEGORY(LogVideoLiveLinkDevice);
 
 
@@ -18,10 +20,10 @@ FVideoDeviceThumbnailExtractor::FVideoDeviceThumbnailExtractor()
 {
 }
 
-TArray<FColor> FVideoDeviceThumbnailExtractor::ExtractThumbnail(UE::CaptureManager::FMediaTextureSample* Sample)
+TArray<FColor> FVideoDeviceThumbnailExtractor::ExtractThumbnail(UE::CaptureManager::FMediaTextureSample* InSample)
 {
-	EMediaTexturePixelFormat SampleFormat = Sample->CurrentFormat;
-	TArray<uint8>& Buffer = Sample->Buffer;
+	EMediaTexturePixelFormat SampleFormat = InSample->CurrentFormat;
+	TArray<uint8>& Buffer = InSample->Buffer;
 
 	TArray<FColor> ThumbnailRawColorData;
 	switch (SampleFormat)
@@ -29,9 +31,9 @@ TArray<FColor> FVideoDeviceThumbnailExtractor::ExtractThumbnail(UE::CaptureManag
 	case EMediaTexturePixelFormat::U8_Mono:
 	{
 		ThumbnailRawColorData.Reserve(Buffer.Num());
-		for (uint8& value : Buffer)
+		for (uint8& Value : Buffer)
 		{
-			ThumbnailRawColorData.Add({ value, value, value, 0 });
+			ThumbnailRawColorData.Add(FColor(Value, Value, Value));
 		}
 		break;
 	}
@@ -40,7 +42,7 @@ TArray<FColor> FVideoDeviceThumbnailExtractor::ExtractThumbnail(UE::CaptureManag
 		ThumbnailRawColorData.Reserve(Buffer.Num() / 3);
 		for (int32 ValueIndex = 0; ValueIndex < Buffer.Num(); ValueIndex += 3)
 		{
-			ThumbnailRawColorData.Add({ Buffer[ValueIndex], Buffer[ValueIndex + 1], Buffer[ValueIndex + 2], 0 });
+			ThumbnailRawColorData.Add(FColor(Buffer[ValueIndex], Buffer[ValueIndex + 1], Buffer[ValueIndex + 2]));
 		}
 		break;
 	}
@@ -49,7 +51,7 @@ TArray<FColor> FVideoDeviceThumbnailExtractor::ExtractThumbnail(UE::CaptureManag
 		ThumbnailRawColorData.Reserve(Buffer.Num() / 3);
 		for (int32 ValueIndex = 0; ValueIndex < Buffer.Num(); ValueIndex += 3)
 		{
-			ThumbnailRawColorData.Add({ Buffer[ValueIndex + 2], Buffer[ValueIndex + 1], Buffer[ValueIndex], 0 });
+			ThumbnailRawColorData.Add(FColor(Buffer[ValueIndex + 2], Buffer[ValueIndex + 1], Buffer[ValueIndex]));
 		}
 		break;
 	}
@@ -58,7 +60,7 @@ TArray<FColor> FVideoDeviceThumbnailExtractor::ExtractThumbnail(UE::CaptureManag
 		ThumbnailRawColorData.Reserve(Buffer.Num() / 4);
 		for (int32 ValueIndex = 0; ValueIndex < Buffer.Num(); ValueIndex += 4)
 		{
-			ThumbnailRawColorData.Add({ Buffer[ValueIndex], Buffer[ValueIndex + 1], Buffer[ValueIndex + 2], Buffer[ValueIndex + 3] });
+			ThumbnailRawColorData.Add(FColor(Buffer[ValueIndex], Buffer[ValueIndex + 1], Buffer[ValueIndex + 2], Buffer[ValueIndex + 3]));
 		}
 		break;
 	}
@@ -67,16 +69,26 @@ TArray<FColor> FVideoDeviceThumbnailExtractor::ExtractThumbnail(UE::CaptureManag
 		ThumbnailRawColorData.Reserve(Buffer.Num() / 4);
 		for (int32 ValueIndex = 0; ValueIndex < Buffer.Num(); ValueIndex += 4)
 		{
-			ThumbnailRawColorData.Add({ Buffer[ValueIndex + 2], Buffer[ValueIndex + 1], Buffer[ValueIndex], Buffer[ValueIndex + 3] });
+			ThumbnailRawColorData.Add(FColor(Buffer[ValueIndex + 2], Buffer[ValueIndex + 1], Buffer[ValueIndex], Buffer[ValueIndex + 3]));
 		}
 		break;
 	}
+	case EMediaTexturePixelFormat::U8_I420:
+		ThumbnailRawColorData = UE::CaptureManager::UEConvertI420ToBGRA(InSample);
+		break;
+	case EMediaTexturePixelFormat::U8_NV12:
+		ThumbnailRawColorData = UE::CaptureManager::UEConvertNV12ToBGRA(InSample);
+		break;
+	case EMediaTexturePixelFormat::U8_YUY2:
+		ThumbnailRawColorData = UE::CaptureManager::UEConvertYUY2ToBGRA(InSample);
+		break;
 	default:
 	{
 		UE_LOG(LogVideoLiveLinkDevice, Warning, TEXT("Unsupported image format"));
 	}
 	}
 
+	InSample->Buffer.Empty();
 	return ThumbnailRawColorData;
 }
 
diff --git a/Engine/Plugins/VirtualProduction/CaptureManager/CaptureManagerDevices/Source/VideoLiveLinkDeviceCommon/VideoLiveLinkDeviceCommon.Build.cs b/Engine/Plugins/VirtualProduction/CaptureManager/CaptureManagerDevices/Source/VideoLiveLinkDeviceCommon/VideoLiveLinkDeviceCommon.Build.cs
--- a/Engine/Plugins/VirtualProduction/CaptureManager/CaptureManagerDevices/Source/VideoLiveLinkDeviceCommon/VideoLiveLinkDeviceCommon.Build.cs
+++ b/Engine/Plugins/VirtualProduction/CaptureManager/CaptureManagerDevices/Source/VideoLiveLinkDeviceCommon/VideoLiveLinkDeviceCommon.Build.cs
@@ -12,6 +12,7 @@ public class VideoLiveLinkDeviceCommon : ModuleRules
 
 		PublicDependencyModuleNames.AddRange(new string[]
 		{
+			"ImageCore",
 			"Core"
 		});
 
