Splitting the download from the convert/upload

diff --git a/Engine/Plugins/VirtualProduction/CaptureManager/CaptureManagerApp/Source/CaptureManagerEditor/Private/IngestManagement/IngestJob.cpp b/Engine/Plugins/VirtualProduction/CaptureManager/CaptureManagerApp/Source/CaptureManagerEditor/Private/IngestManagement/IngestJob.cpp
--- a/Engine/Plugins/VirtualProduction/CaptureManager/CaptureManagerApp/Source/CaptureManagerEditor/Private/IngestManagement/IngestJob.cpp
+++ b/Engine/Plugins/VirtualProduction/CaptureManager/CaptureManagerApp/Source/CaptureManagerEditor/Private/IngestManagement/IngestJob.cpp
@@ -9,12 +9,14 @@ FIngestJob::FIngestJob(
 	FGuid InCaptureDeviceId,
 	FTakeId InTakeId,
 	FTakeMetadata InTakeMetadata,
+	EIngestCapability_ProcessConfig InPipelineConfig,
 	FSettings InSettings
 ) :
 	JobGuid(FGuid::NewGuid()),
 	CaptureDeviceId(InCaptureDeviceId),
 	TakeId(InTakeId),
 	TakeMetadata(MoveTemp(InTakeMetadata)),
+	PipelineConfig(InPipelineConfig),
 	Settings(MoveTemp(InSettings))
 {
 }
@@ -55,6 +57,11 @@ const FIngestJob::FSettings& FIngestJob::GetSettings() const
 	return Settings;
 }
 
+EIngestCapability_ProcessConfig FIngestJob::GetPipelineConfig() const
+{
+	return PipelineConfig;
+}
+
 void FIngestJob::SetProcessingState(const FIngestJob::EProcessingState InProcessingState)
 {
 	ProcessingState = InProcessingState;
diff --git a/Engine/Plugins/VirtualProduction/CaptureManager/CaptureManagerApp/Source/CaptureManagerEditor/Private/IngestManagement/IngestJob.h b/Engine/Plugins/VirtualProduction/CaptureManager/CaptureManagerApp/Source/CaptureManagerEditor/Private/IngestManagement/IngestJob.h
--- a/Engine/Plugins/VirtualProduction/CaptureManager/CaptureManagerApp/Source/CaptureManagerEditor/Private/IngestManagement/IngestJob.h
+++ b/Engine/Plugins/VirtualProduction/CaptureManager/CaptureManagerApp/Source/CaptureManagerEditor/Private/IngestManagement/IngestJob.h
@@ -3,6 +3,7 @@
 #pragma once
 
 #include "Ingest/LiveLinkDeviceCapability_Ingest.h"
+#include "Ingest/IngestCapability_ProcessHandle.h"
 
 #include "UIngestJobSettings.h"
 
@@ -47,6 +48,7 @@ public:
 		FGuid InCaptureDeviceId,
 		FTakeId InTakeId,
 		FTakeMetadata InTakeMetadata,
+		EIngestCapability_ProcessConfig InPipelineConfig,
 		FSettings InSettings
 	);
 
@@ -54,6 +56,7 @@ public:
 	FGuid GetCaptureDeviceId() const;
 	FTakeId GetTakeId() const;
 	const FTakeMetadata& GetTakeMetadata() const;
+	EIngestCapability_ProcessConfig GetPipelineConfig() const;
 	const FSettings& GetSettings() const;
 	EProcessingState GetProcessingState() const;
 	double GetProgress() const;
@@ -67,6 +70,7 @@ private:
 	const FGuid CaptureDeviceId;
 	const FTakeId TakeId;
 	const FTakeMetadata TakeMetadata;
+	const EIngestCapability_ProcessConfig PipelineConfig;
 
 	mutable FCriticalSection CriticalSection;
 	FSettings Settings;
diff --git a/Engine/Plugins/VirtualProduction/CaptureManager/CaptureManagerApp/Source/CaptureManagerEditor/Private/IngestManagement/IngestJobExecutor.cpp b/Engine/Plugins/VirtualProduction/CaptureManager/CaptureManagerApp/Source/CaptureManagerEditor/Private/IngestManagement/IngestJobExecutor.cpp
--- a/Engine/Plugins/VirtualProduction/CaptureManager/CaptureManagerApp/Source/CaptureManagerEditor/Private/IngestManagement/IngestJobExecutor.cpp
+++ b/Engine/Plugins/VirtualProduction/CaptureManager/CaptureManagerApp/Source/CaptureManagerEditor/Private/IngestManagement/IngestJobExecutor.cpp
@@ -24,7 +24,7 @@ DEFINE_LOG_CATEGORY_STATIC(LogCaptureManagerIngestJobExecutor, Log, All);
 namespace UE::CaptureManager
 {
 
-EIngestCapability_ImagePixelFormat ConvertImagePixelFormat(::EImagePixelFormat InImagePixelFormat)
+static EIngestCapability_ImagePixelFormat ConvertImagePixelFormat(::EImagePixelFormat InImagePixelFormat)
 {
 	switch (InImagePixelFormat)
 	{
@@ -36,7 +36,7 @@ EIngestCapability_ImagePixelFormat ConvertImagePixelFormat(::EImagePixelFormat I
 	}
 }
 
-EIngestCapability_ImageRotation ConvertImageRotation(::EImageRotation InImageRotation)
+static EIngestCapability_ImageRotation ConvertImageRotation(::EImageRotation InImageRotation)
 {
 	switch (InImageRotation)
 	{
@@ -93,18 +93,18 @@ uint32 FIngestJobExecutor::Run()
 
 			// Note: This may be a blocking call depending on how the capture device has been implemented. If so, this 
 			// may take a while to stop (no way to cancel in that case).
-			UIngestCapability_ProcessHandle* TakeProcessHandle = RunIngest(*Job, Promise);
+			TUniquePtr<FIngestJobExecutor::FRunIngestContext> RunIngestContext = RunIngest(*Job, Promise);
 
 			FScopeLock Lock(&CriticalSection);
-			check(!IngestProcessHandle);
+			check(RunIngestContext);
 
-			IngestProcessHandle = MoveTemp(TakeProcessHandle);
+			RunIngestProcessContext = MoveTemp(RunIngestContext);
 
-			// There's a period between RunIngest and the assignment of the member handle, during which Stop() may have
+			// There's a period between RunIngest and the assignment of the member context, during which Stop() may have
 			// been called. If so, immediately terminate (as Cancel() will not be called otherwise).
-			if (bStopRequested && !Future.IsReady() && IngestProcessHandle)
+			if (bStopRequested && !Future.IsReady() && RunIngestProcessContext)
 			{
-				IngestProcessHandle->Cancel();
+				ILiveLinkDeviceCapability_Ingest::Execute_CancelIngestProcess(RunIngestProcessContext->Device.Get(), RunIngestProcessContext->ProcessHandle.Get());
 			}
 		}
 		else
@@ -119,9 +119,9 @@ uint32 FIngestJobExecutor::Run()
 		{
 			FScopeLock Lock(&CriticalSection);
 
-			if (IngestProcessHandle)
+			if (RunIngestProcessContext)
 			{
-				IngestProcessHandle = nullptr;
+				RunIngestProcessContext = nullptr;
 			}
 		}
 
@@ -140,9 +140,9 @@ void FIngestJobExecutor::Stop()
 
 	FScopeLock Lock(&CriticalSection);
 
-	if (IsValid(IngestProcessHandle))
+	if (RunIngestProcessContext.IsValid())
 	{
-		IngestProcessHandle->Cancel();
+		ILiveLinkDeviceCapability_Ingest::Execute_CancelIngestProcess(RunIngestProcessContext->Device.Get(), RunIngestProcessContext->ProcessHandle.Get());
 	}
 }
 
@@ -152,7 +152,7 @@ void FIngestJobExecutor::Exit()
 	Complete.ExecuteIfBound();
 }
 
-UIngestCapability_ProcessHandle* FIngestJobExecutor::RunIngest(FIngestJob& InJob, TPromise<void>& InPromise)
+TUniquePtr<FIngestJobExecutor::FRunIngestContext> FIngestJobExecutor::RunIngest(FIngestJob& InJob, TPromise<void>& InPromise)
 {
 	if (!GEditor)
 	{
@@ -252,7 +252,7 @@ UIngestCapability_ProcessHandle* FIngestJobExecutor::RunIngest(FIngestJob& InJob
 		UE_LOG(
 			LogCaptureManagerIngestJobExecutor,
 			Error,
-			TEXT("[%s] Ingest failed, Capture Device does not support the ingest capability: %s"), *ExecutorName, *Device.GetName()
+			TEXT("[%s] Ingest failed, Capture Device does not support the ingest capability: %s"), *ExecutorName, *(Device->GetName())
 		);
 
 		// Set the processing state to anything other than Pending, to prevent other executors picking up the same bad job.
@@ -262,29 +262,40 @@ UIngestCapability_ProcessHandle* FIngestJobExecutor::RunIngest(FIngestJob& InJob
 	}
 
 	FIngestProcessFinishReporter ProcessFinished = FIngestProcessFinishReporter::Type::CreateLambda(
-		[this, &InPromise, &InJob](FTakeId TakeId, TValueOrError<void, FIngestCapability_Error> IngestError)
+		[this, &InPromise, &InJob](const UIngestCapability_ProcessHandle* InProcessHandle, TValueOrError<void, FIngestCapability_Error> IngestError)
 		{
 			FIngestJob::EProcessingState ProcessingState;
 
-			if (IngestError.HasValue())
+			bool bIsDone = InProcessHandle->IsDone();
+
+			if (IngestError.HasError())
+			{
+				ProcessingState = FIngestJob::EProcessingState::Aborted;
+				UE_LOG(LogCaptureManagerIngestJobExecutor, Error, TEXT("[%s] Ingest failed: %s (Error code: %d)"), *ExecutorName, *IngestError.GetError().GetMessage(), IngestError.GetError().GetCode());
+				bIsDone = true;
+			}
+			else if (bIsDone)
 			{
 				ProcessingState = FIngestJob::EProcessingState::Complete;
 				InJob.SetProgress(1.0);
 			}
 			else
 			{
-				ProcessingState = FIngestJob::EProcessingState::Aborted;
-				UE_LOG(LogCaptureManagerIngestJobExecutor, Error, TEXT("[%s] Ingest failed: %s (Error code: %d)"), *ExecutorName, *IngestError.GetError().GetMessage(), IngestError.GetError().GetCode())
+				ProcessingState = FIngestJob::EProcessingState::Running;
 			}
 
 			InJob.SetProcessingState(ProcessingState);
 			JobProcessingStateChanged.ExecuteIfBound(InJob.GetGuid(), ProcessingState);
-			InPromise.SetValue();
+			
+			if (bIsDone)
+			{
+				InPromise.SetValue();
+			}
 		}
 	);
 
 	FIngestProcessProgressReporter ProgressReporter = FIngestProcessProgressReporter::Type::CreateLambda(
-		[&InJob, LastProgress = 0.0](FTakeId TakeId, const double Progress) mutable
+		[&InJob, LastProgress = 0.0](const UIngestCapability_ProcessHandle* InProcessHandle, const double Progress) mutable
 		{
 			// A little sanity check that we're receiving updates in the expected order.
 			check(Progress >= LastProgress);
@@ -294,14 +305,18 @@ UIngestCapability_ProcessHandle* FIngestJobExecutor::RunIngest(FIngestJob& InJob
 		}
 	);
 
-	UIngestCapability_ProcessHandle* ProcessHandle =
-		ILiveLinkDeviceCapability_Ingest::Execute_PrepareIngestProcess(Device, InJob.GetTakeId());
+	TStrongObjectPtr<UIngestCapability_ProcessHandle> ProcessHandle(
+		ILiveLinkDeviceCapability_Ingest::Execute_CreateIngestProcess(Device, InJob.GetTakeId(), InJob.GetPipelineConfig()));
+
+	ProcessHandle->OnProcessFinishReporter() = MoveTemp(ProcessFinished);
+	ProcessHandle->OnProcessProgressReporter() = MoveTemp(ProgressReporter);
 
-	ProcessHandle->ProcessFinishedReporter = MoveTemp(ProcessFinished);
-	ProcessHandle->ProcessProgressReporter = MoveTemp(ProgressReporter);
+	ILiveLinkDeviceCapability_Ingest::Execute_RunIngestProcess(Device, ProcessHandle.Get(), IngestOptions);
 
-	ProcessHandle->Run(IngestOptions);
-	return ProcessHandle;
+	TUniquePtr<FRunIngestContext> Context = MakeUnique<FRunIngestContext>();
+	Context->Device = TStrongObjectPtr<ULiveLinkDevice>(Device);
+	Context->ProcessHandle = MoveTemp(ProcessHandle);
+	return Context;
 }
 
 } // namespace UE::CaptureManager
diff --git a/Engine/Plugins/VirtualProduction/CaptureManager/CaptureManagerApp/Source/CaptureManagerEditor/Private/IngestManagement/IngestJobExecutor.h b/Engine/Plugins/VirtualProduction/CaptureManager/CaptureManagerApp/Source/CaptureManagerEditor/Private/IngestManagement/IngestJobExecutor.h
--- a/Engine/Plugins/VirtualProduction/CaptureManager/CaptureManagerApp/Source/CaptureManagerEditor/Private/IngestManagement/IngestJobExecutor.h
+++ b/Engine/Plugins/VirtualProduction/CaptureManager/CaptureManagerApp/Source/CaptureManagerEditor/Private/IngestManagement/IngestJobExecutor.h
@@ -25,12 +25,18 @@ public:
 	virtual void Exit() override;
 
 private:
-	UIngestCapability_ProcessHandle* RunIngest(FIngestJob& InJob, TPromise<void>& InPromise);
+	struct FRunIngestContext
+	{
+		TStrongObjectPtr<ULiveLinkDevice> Device;
+		TStrongObjectPtr<UIngestCapability_ProcessHandle> ProcessHandle;
+	};
+
+	TUniquePtr<FIngestJobExecutor::FRunIngestContext> RunIngest(FIngestJob& InJob, TPromise<void>& InPromise);
 
 	FCriticalSection CriticalSection;
 	FString ExecutorName;
 	TSharedRef<FIngestJobQueue> ProcessingQueue;
-	UIngestCapability_ProcessHandle* IngestProcessHandle = nullptr;
+	TUniquePtr<FIngestJobExecutor::FRunIngestContext> RunIngestProcessContext;
 	FOnComplete Complete;
 	FJobProcessingStateChanged JobProcessingStateChanged;
 
diff --git a/Engine/Plugins/VirtualProduction/CaptureManager/CaptureManagerApp/Source/CaptureManagerEditor/Private/IngestManagement/IngestPipelineManager.cpp b/Engine/Plugins/VirtualProduction/CaptureManager/CaptureManagerApp/Source/CaptureManagerEditor/Private/IngestManagement/IngestPipelineManager.cpp
--- a/Engine/Plugins/VirtualProduction/CaptureManager/CaptureManagerApp/Source/CaptureManagerEditor/Private/IngestManagement/IngestPipelineManager.cpp
+++ b/Engine/Plugins/VirtualProduction/CaptureManager/CaptureManagerApp/Source/CaptureManagerEditor/Private/IngestManagement/IngestPipelineManager.cpp
@@ -11,12 +11,34 @@ namespace UE::CaptureManager
 
 FIngestPipelineManager::FIngestPipelineManager()
 {
-	Pipelines.Emplace(
-		LOCTEXT("DefaultPipelineDisplayName","Ingest"),
-		LOCTEXT("DefaultPipelineToolTip", "Download, Convert, Create Assets")
-	);
+	UEnum* ProcessConfig = StaticEnum<EIngestCapability_ProcessConfig>();
+
+	for (int32 EnumIndex = 0; EnumIndex < ProcessConfig->NumEnums() - 1; ++EnumIndex) // -1 to skip the implicit "MAX"
+	{
+		if (ProcessConfig->HasMetaData(TEXT("Hidden"), EnumIndex))
+		{
+			continue;
+		}
+
+		FText DisplayName = ProcessConfig->GetDisplayNameTextByIndex(EnumIndex);
+		FText Tooltip = ProcessConfig->GetToolTipTextByIndex(EnumIndex);
+		EIngestCapability_ProcessConfig PipelineConfig = 
+			static_cast<EIngestCapability_ProcessConfig>(ProcessConfig->GetValueByIndex(EnumIndex));
 
-	SelectedPipeline = *Pipelines.CreateIterator();
+		FPipelineDetails PipelineDetails =
+		{
+			.DisplayName = MoveTemp(DisplayName),
+			.ToolTip = MoveTemp(Tooltip),
+			.PipelineConfig = PipelineConfig
+		};
+
+		if (PipelineConfig == EIngestCapability_ProcessConfig::Ingest)
+		{
+			SelectedPipeline = PipelineDetails;
+		}
+
+		Pipelines.Emplace(MoveTemp(PipelineDetails));
+	}
 }
 
 const TArray<FPipelineDetails>& FIngestPipelineManager::GetPipelines() const
diff --git a/Engine/Plugins/VirtualProduction/CaptureManager/CaptureManagerApp/Source/CaptureManagerEditor/Private/IngestManagement/IngestPipelineManager.h b/Engine/Plugins/VirtualProduction/CaptureManager/CaptureManagerApp/Source/CaptureManagerEditor/Private/IngestManagement/IngestPipelineManager.h
--- a/Engine/Plugins/VirtualProduction/CaptureManager/CaptureManagerApp/Source/CaptureManagerEditor/Private/IngestManagement/IngestPipelineManager.h
+++ b/Engine/Plugins/VirtualProduction/CaptureManager/CaptureManagerApp/Source/CaptureManagerEditor/Private/IngestManagement/IngestPipelineManager.h
@@ -5,6 +5,8 @@
 #include "Containers/Array.h"
 #include "Internationalization/Text.h"
 
+#include "Ingest/IngestCapability_ProcessHandle.h"
+
 namespace UE::CaptureManager
 {
 
@@ -12,6 +14,7 @@ struct FPipelineDetails
 {
 	FText DisplayName;
 	FText ToolTip;
+	EIngestCapability_ProcessConfig PipelineConfig;
 };
 
 class FIngestPipelineManager
diff --git a/Engine/Plugins/VirtualProduction/CaptureManager/CaptureManagerApp/Source/CaptureManagerEditor/Private/Scripting/ULiveLinkHubCaptureDevice.cpp b/Engine/Plugins/VirtualProduction/CaptureManager/CaptureManagerApp/Source/CaptureManagerEditor/Private/Scripting/ULiveLinkHubCaptureDevice.cpp
--- a/Engine/Plugins/VirtualProduction/CaptureManager/CaptureManagerApp/Source/CaptureManagerEditor/Private/Scripting/ULiveLinkHubCaptureDevice.cpp
+++ b/Engine/Plugins/VirtualProduction/CaptureManager/CaptureManagerApp/Source/CaptureManagerEditor/Private/Scripting/ULiveLinkHubCaptureDevice.cpp
@@ -39,6 +39,9 @@ struct ULiveLinkHubCaptureDevice::FImpl
 
 	TObjectPtr<UIngestCapability_ProcessResult> ImportTake(const FLiveLinkHubTakeMetadata& InTake, 
 														   const TObjectPtr<const UIngestCapability_Options> InConversionSettings) const;
+	TObjectPtr<UIngestCapability_ProcessResult> DownloadTake(const FLiveLinkHubTakeMetadata& InTake,
+															 const TObjectPtr<const UIngestCapability_Options> InConversionSettings) const;
+
 	TObjectPtr<UIngestCapability_ProcessResult> FetchTakes(TArray<FLiveLinkHubTakeMetadata>& OutTakes) const;
 
 	TObjectPtr<UIngestCapability_ProcessResult> Start(const int32 InTimeoutSeconds);
@@ -51,6 +54,12 @@ struct ULiveLinkHubCaptureDevice::FImpl
 	TOptional<FSession> Session;
 
 	TObjectPtr<ULiveLinkDevice> BackgroundDevice;
+
+private:
+
+	TObjectPtr<UIngestCapability_ProcessResult> RunProcessTake(const FLiveLinkHubTakeMetadata& InTake,
+															   const TObjectPtr<const UIngestCapability_Options> InConversionSettings, 
+															   EIngestCapability_ProcessConfig InProcessConfig) const;
 };
 
 ULiveLinkHubCaptureDevice::FImpl::FSession::FSession(const TObjectPtr<ULiveLinkDevice> InDevice) :
@@ -173,7 +182,20 @@ void ULiveLinkHubCaptureDevice::FImpl::RemoveDevice()
 }
 
 TObjectPtr<UIngestCapability_ProcessResult> ULiveLinkHubCaptureDevice::FImpl::ImportTake(const FLiveLinkHubTakeMetadata& InTake, 
-																						 const TObjectPtr<const UIngestCapability_Options> InConversionSettings) const
+																						 TObjectPtr<const UIngestCapability_Options> InConversionSettings) const
+{
+	return RunProcessTake(InTake, MoveTemp(InConversionSettings), EIngestCapability_ProcessConfig::Ingest);
+}
+
+TObjectPtr<UIngestCapability_ProcessResult> ULiveLinkHubCaptureDevice::FImpl::DownloadTake(const FLiveLinkHubTakeMetadata& InTake,
+																						   TObjectPtr<const UIngestCapability_Options> InConversionSettings) const
+{
+	return RunProcessTake(InTake, MoveTemp(InConversionSettings), EIngestCapability_ProcessConfig::Download);
+}
+
+TObjectPtr<UIngestCapability_ProcessResult> ULiveLinkHubCaptureDevice::FImpl::RunProcessTake(const FLiveLinkHubTakeMetadata& InTake,
+																							 const TObjectPtr<const UIngestCapability_Options> InConversionSettings,
+																							 EIngestCapability_ProcessConfig InProcessConfig) const
 {
 	using namespace UE::CaptureManager;
 
@@ -196,29 +218,36 @@ TObjectPtr<UIngestCapability_ProcessResult> ULiveLinkHubCaptureDevice::FImpl::Im
 	bool bImportSuccess = false;
 	FText ErrorMessage;
 
-	TManagedDelegate<FTakeId, TValueOrError<void, FIngestCapability_Error>> OnIngestFinished(
-		[&ImportCompleteEvent, &bImportSuccess, &ErrorMessage](FTakeId InTakeId, TValueOrError<void, FIngestCapability_Error> InImportResult)
+	TManagedDelegate<const UIngestCapability_ProcessHandle*, TValueOrError<void, FIngestCapability_Error>> OnIngestFinished(
+		[&ImportCompleteEvent, &bImportSuccess, &ErrorMessage](const UIngestCapability_ProcessHandle* InProcessHandle, TValueOrError<void, FIngestCapability_Error> InImportResult)
 		{
-			if (InImportResult.HasValue())
+			bool bIsDone = InProcessHandle->IsDone();
+
+			if (InImportResult.HasError())
 			{
-				bImportSuccess = true;
+				ErrorMessage = FText::Format(ImportFailedFormatText, FText::FromString(InImportResult.GetError().GetMessage()));
+				bIsDone = true;
 			}
-			else
+			else if (bIsDone)
 			{
-				ErrorMessage = FText::Format(ImportFailedFormatText, FText::FromString(InImportResult.GetError().GetMessage()));
+				bImportSuccess = true;
 			}
 
-			ImportCompleteEvent->Trigger();
+			if (bIsDone)
+			{
+				ImportCompleteEvent->Trigger();
+			}
 		},
 		EDelegateExecutionThread::InternalThread
 	);
 
-	UIngestCapability_ProcessHandle* ProcessHandle =
-		ILiveLinkDeviceCapability_Ingest::Execute_PrepareIngestProcess(BackgroundDevice, InTake.TakeId);
+	TStrongObjectPtr<UIngestCapability_ProcessHandle> ProcessHandle(
+		ILiveLinkDeviceCapability_Ingest::Execute_CreateIngestProcess(BackgroundDevice, InTake.TakeId, InProcessConfig));
 	check(ProcessHandle);
 
-	ProcessHandle->ProcessFinishedReporter = MoveTemp(OnIngestFinished);
-	ProcessHandle->Run(InConversionSettings);
+	ProcessHandle->OnProcessFinishReporter() = MoveTemp(OnIngestFinished);
+
+	ILiveLinkDeviceCapability_Ingest::Execute_RunIngestProcess(BackgroundDevice, ProcessHandle.Get(), InConversionSettings);
 
 	ImportCompleteEvent->Wait();
 
@@ -328,6 +357,13 @@ UIngestCapability_ProcessResult* ULiveLinkHubCaptureDevice::ImportTake(const FLi
 	return ImplPtr->ImportTake(InTake, MoveTemp(ConversionSettings));
 }
 
+UIngestCapability_ProcessResult* ULiveLinkHubCaptureDevice::DownloadTake(const FLiveLinkHubTakeMetadata& InTake, const FString& InDownloadDirectory) const
+{
+	TObjectPtr<UIngestCapability_Options> ConversionSettings = NewObject<UIngestCapability_Options>();
+	ConversionSettings->DownloadDirectory = InDownloadDirectory;
+	return ImplPtr->DownloadTake(InTake, MoveTemp(ConversionSettings));
+}
+
 FLiveLinkHubFetchTakesResult ULiveLinkHubCaptureDevice::FetchTakes() const
 {
 	TArray<FLiveLinkHubTakeMetadata> Takes;
diff --git a/Engine/Plugins/VirtualProduction/CaptureManager/CaptureManagerApp/Source/CaptureManagerEditor/Private/Scripting/ULiveLinkHubCaptureDevice.h b/Engine/Plugins/VirtualProduction/CaptureManager/CaptureManagerApp/Source/CaptureManagerEditor/Private/Scripting/ULiveLinkHubCaptureDevice.h
--- a/Engine/Plugins/VirtualProduction/CaptureManager/CaptureManagerApp/Source/CaptureManagerEditor/Private/Scripting/ULiveLinkHubCaptureDevice.h
+++ b/Engine/Plugins/VirtualProduction/CaptureManager/CaptureManagerApp/Source/CaptureManagerEditor/Private/Scripting/ULiveLinkHubCaptureDevice.h
@@ -67,6 +67,9 @@ public:
 	UFUNCTION(BlueprintCallable, Category = Ingest)
 	UIngestCapability_ProcessResult* ImportTake(const FLiveLinkHubTakeMetadata& InTake, const UIngestCapability_Options* InConversionSettings) const;
 
+	UFUNCTION(BlueprintCallable, Category = Ingest)
+	UIngestCapability_ProcessResult* DownloadTake(const FLiveLinkHubTakeMetadata& InTake, const FString& InDownloadDirectory) const;
+
 	UFUNCTION(BlueprintCallable, Category = Ingest)
 	FLiveLinkHubFetchTakesResult FetchTakes() const;
 
diff --git a/Engine/Plugins/VirtualProduction/CaptureManager/CaptureManagerApp/Source/CaptureManagerEditor/Private/UI/SIngestJobProcessor.cpp b/Engine/Plugins/VirtualProduction/CaptureManager/CaptureManagerApp/Source/CaptureManagerEditor/Private/UI/SIngestJobProcessor.cpp
--- a/Engine/Plugins/VirtualProduction/CaptureManager/CaptureManagerApp/Source/CaptureManagerEditor/Private/UI/SIngestJobProcessor.cpp
+++ b/Engine/Plugins/VirtualProduction/CaptureManager/CaptureManagerApp/Source/CaptureManagerEditor/Private/UI/SIngestJobProcessor.cpp
@@ -38,6 +38,7 @@ void ResetJob(TSharedPtr<FIngestJob> InIngestJob, TSharedPtr<FIngestJobProcessor
 			InIngestJob->GetCaptureDeviceId(),
 			InIngestJob->GetTakeId(),
 			InIngestJob->GetTakeMetadata(),
+			InIngestJob->GetPipelineConfig(),
 			InIngestJob->GetSettings()
 		);
 
diff --git a/Engine/Plugins/VirtualProduction/CaptureManager/CaptureManagerApp/Source/CaptureManagerEditor/Private/UI/TakesView.cpp b/Engine/Plugins/VirtualProduction/CaptureManager/CaptureManagerApp/Source/CaptureManagerEditor/Private/UI/TakesView.cpp
--- a/Engine/Plugins/VirtualProduction/CaptureManager/CaptureManagerApp/Source/CaptureManagerEditor/Private/UI/TakesView.cpp
+++ b/Engine/Plugins/VirtualProduction/CaptureManager/CaptureManagerApp/Source/CaptureManagerEditor/Private/UI/TakesView.cpp
@@ -386,6 +386,7 @@ void FTakesView::AddTakesToIngestQueue(const TArray<TStrongObjectPtr<UTakeVirtua
 	TArray<TSharedRef<FIngestJob>> IngestJobs;
 	IngestJobs.Reserve(InTakeAssets.Num());
 
+	const FPipelineDetails Pipeline = Controller->GetIngestPipelineManager()->GetSelectedPipeline();
 	for (const TStrongObjectPtr<UTakeVirtualAsset>& TakeAsset : InTakeAssets)
 	{
 
@@ -412,7 +413,7 @@ void FTakesView::AddTakesToIngestQueue(const TArray<TStrongObjectPtr<UTakeVirtua
 			.UploadHostName = InJobPreset.UploadHostName
 		};
 
-		IngestJobs.Emplace(MakeShared<FIngestJob>(TakeAsset->CaptureDeviceId, TakeAsset->TakeId, TakeAsset->Metadata, MoveTemp(IngestJobSettings)));
+		IngestJobs.Emplace(MakeShared<FIngestJob>(TakeAsset->CaptureDeviceId, TakeAsset->TakeId, TakeAsset->Metadata, Pipeline.PipelineConfig, MoveTemp(IngestJobSettings)));
 	}
 
 	const int32 ExpectedNumAdded = IngestJobs.Num();
diff --git a/Engine/Plugins/VirtualProduction/CaptureManager/CaptureManagerApp/Source/ExampleLiveLinkDevices/Private/NetworkIngest/ExampleNetworkIngestDevice.cpp b/Engine/Plugins/VirtualProduction/CaptureManager/CaptureManagerApp/Source/ExampleLiveLinkDevices/Private/NetworkIngest/ExampleNetworkIngestDevice.cpp
--- a/Engine/Plugins/VirtualProduction/CaptureManager/CaptureManagerApp/Source/ExampleLiveLinkDevices/Private/NetworkIngest/ExampleNetworkIngestDevice.cpp
+++ b/Engine/Plugins/VirtualProduction/CaptureManager/CaptureManagerApp/Source/ExampleLiveLinkDevices/Private/NetworkIngest/ExampleNetworkIngestDevice.cpp
@@ -80,7 +80,7 @@ void UExampleNetworkIngestDevice::UpdateTakeList_Implementation(UIngestCapabilit
 	// See ULiveLinkFaceDevice::UpdateTakeList_Implementation
 }
 
-void UExampleNetworkIngestDevice::RunIngestTake(int32 InTakeId, const UIngestCapability_Options* InIngestOptions)
+void UExampleNetworkIngestDevice::RunIngestTake(const UIngestCapability_ProcessHandle* InProcessHandle, const UIngestCapability_Options* InIngestOptions)
 {
 	// Get the take corresponding to the InTakeId from your device (i.e. Download the take from the device)
 	//     See ULiveLinkFaceDevice::RunIngestTake
@@ -91,21 +91,21 @@ void UExampleNetworkIngestDevice::RunIngestTake(int32 InTakeId, const UIngestCap
 
 	using namespace UE::CaptureManager;
 	TSharedPtr<FTaskProgress> TaskProgress = MakeShared<FTaskProgress>
-		(NumberOfTasks, FTaskProgress::FProgressReporter::CreateLambda([this, TakeId = InTakeId](double InProgress)
+		(NumberOfTasks, FTaskProgress::FProgressReporter::CreateLambda([this, TakeId = InProcessHandle->TakeId](double InProgress)
 			{
 				ExecuteProcessProgressReporter(TakeId, InProgress);
 			}));
 
 	// Note: IngestTake uses GetFullTakePath() so make sure you have implemented it
 
-	Super::IngestTake(InTakeId, InIngestOptions, MoveTemp(TaskProgress));
+	Super::IngestTake(InProcessHandle, InIngestOptions, MoveTemp(TaskProgress));
 }
 
-void UExampleNetworkIngestDevice::CancelIngest(int32 InTakeId)
+void UExampleNetworkIngestDevice::CancelIngestProcess_Implementation(const UIngestCapability_ProcessHandle* InProcessHandle)
 {
 	// Stop your device from sending data here
 
-	Super::CancelIngest(InTakeId);
+	Super::CancelIngest(InProcessHandle->TakeId);
 }
 
 ELiveLinkDeviceConnectionStatus UExampleNetworkIngestDevice::GetConnectionStatus_Implementation() const
diff --git a/Engine/Plugins/VirtualProduction/CaptureManager/CaptureManagerApp/Source/ExampleLiveLinkDevices/Private/NetworkIngest/ExampleNetworkIngestDevice.h b/Engine/Plugins/VirtualProduction/CaptureManager/CaptureManagerApp/Source/ExampleLiveLinkDevices/Private/NetworkIngest/ExampleNetworkIngestDevice.h
--- a/Engine/Plugins/VirtualProduction/CaptureManager/CaptureManagerApp/Source/ExampleLiveLinkDevices/Private/NetworkIngest/ExampleNetworkIngestDevice.h
+++ b/Engine/Plugins/VirtualProduction/CaptureManager/CaptureManagerApp/Source/ExampleLiveLinkDevices/Private/NetworkIngest/ExampleNetworkIngestDevice.h
@@ -57,9 +57,9 @@ private:
 
 	//~ Begin ULiveLinkDeviceCapability_Ingest interface
 	virtual void UpdateTakeList_Implementation(UIngestCapability_UpdateTakeListCallback* InCallback) override;
-	virtual void RunIngestTake(int32 InTakeId,
-									const UIngestCapability_Options* InIngestOptions) override;
-	virtual void CancelIngest(int32 InTakeId) override;
+	virtual void RunIngestTake(const UIngestCapability_ProcessHandle* InProcessHandle,
+							   const UIngestCapability_Options* InIngestOptions) override;
+	virtual void CancelIngestProcess_Implementation(const UIngestCapability_ProcessHandle* InProcessHandle) override;
 	//~ End ULiveLinkDeviceCapability_Ingest interface
 
 	//~ Begin ILiveLinkDeviceCapability_Connection interface
diff --git a/Engine/Plugins/VirtualProduction/CaptureManager/CaptureManagerApp/Source/IngestLiveLinkDevice/Private/BaseIngestLiveLinkDevice.cpp b/Engine/Plugins/VirtualProduction/CaptureManager/CaptureManagerApp/Source/IngestLiveLinkDevice/Private/BaseIngestLiveLinkDevice.cpp
--- a/Engine/Plugins/VirtualProduction/CaptureManager/CaptureManagerApp/Source/IngestLiveLinkDevice/Private/BaseIngestLiveLinkDevice.cpp
+++ b/Engine/Plugins/VirtualProduction/CaptureManager/CaptureManagerApp/Source/IngestLiveLinkDevice/Private/BaseIngestLiveLinkDevice.cpp
@@ -11,6 +11,8 @@
 
 #include "CaptureDataConverter.h"
 
+#include "Async/TaskWaiter.h"
+
 DEFINE_LOG_CATEGORY(LogIngestLiveLinkDevice);
 
 namespace UE::CaptureManager::Private
@@ -30,58 +32,15 @@ static void TryCleanUpWorkingDirectory(const FString& InWorkingDirectory)
 
 } // namespace UE::CaptureManager::Private
 
-class FTaskWaiter {
-public:
-	FTaskWaiter() : TaskCounter(0)
-	{
-
-	}
-
-	~FTaskWaiter()
-	{
-		check(TaskCounter.Load() <= CanCreateTaskFlag);
-	}
-
-	bool CreateTask()
-	{
-		uint32 OldCounter = TaskCounter.IncrementExchange();
-		
-		if ((OldCounter & CanCreateTaskFlag) > 0)
-		{
-			TaskCounter.DecrementExchange();
-			return false;
-		}
-		return true;
-	}
-
-	void FinishTask()
-	{
-		TaskCounter.DecrementExchange();
-	}
-	
-	void WaitForAll()
-	{
-		TaskCounter.AddExchange(CanCreateTaskFlag);
-
-		while (TaskCounter > CanCreateTaskFlag)
-		{
-			FPlatformProcess::Sleep(0.1f);
-		}
-	}
-
-
-private: 
-	TAtomic<uint32> TaskCounter;
-	const uint32 CanCreateTaskFlag = 0x80000000;
-};
-
 class UBaseIngestLiveLinkDevice::FImpl : public TSharedFromThis<UBaseIngestLiveLinkDevice::FImpl>
 {
 public:
 
 	explicit FImpl(UBaseIngestLiveLinkDevice* InOwner);
 
-	void IngestTake(int32 InTakeId, const UIngestCapability_Options* InIngestOptions, TSharedPtr<UE::CaptureManager::FTaskProgress> InTaskProgress);
+	void IngestTake(const UIngestCapability_ProcessHandle* InProcessHandle, 
+					const UIngestCapability_Options* InIngestOptions, 
+					TSharedPtr<UE::CaptureManager::FTaskProgress> InTaskProgress);
 	void CancelIngest(int32 InTakeId);
 	void OnDeviceRemoved(TArray<int32> TakeIdentifiers);
 
@@ -109,7 +68,7 @@ private:
 	// Once function finishes it's work, counter is decremented.
 	// And if it happens that device is being removed before counter is equal to zero, TaskWaiter will busy wait for it to reach zero.
 	// Ensures that 'IngestTake' does not access potentially destroyed device.
-	FTaskWaiter TaskWaiter;
+	UE::CaptureManager::FTaskWaiter TaskWaiter;
 };
 
 UBaseIngestLiveLinkDevice::FImpl::FImpl(UBaseIngestLiveLinkDevice* InOwner) :
@@ -140,7 +99,7 @@ void UBaseIngestLiveLinkDevice::FImpl::OnDeviceRemoved(TArray<int32> TakeIdentif
 	TaskWaiter.WaitForAll();
 }
 
-void UBaseIngestLiveLinkDevice::FImpl::IngestTake(int32 InTakeId, const UIngestCapability_Options* InIngestOptions, TSharedPtr<UE::CaptureManager::FTaskProgress> InTaskProgress)
+void UBaseIngestLiveLinkDevice::FImpl::IngestTake(const UIngestCapability_ProcessHandle* InProcessHandle, const UIngestCapability_Options* InIngestOptions, TSharedPtr<UE::CaptureManager::FTaskProgress> InTaskProgress)
 {
 	if (!TaskWaiter.CreateTask())
 	{
@@ -157,22 +116,23 @@ void UBaseIngestLiveLinkDevice::FImpl::IngestTake(int32 InTakeId, const UIngestC
 	using namespace UE::CaptureManager;
 	using namespace UE::CaptureManager::Private;
 
-	FString TakeFullPath = Owner->GetFullTakePath(InTakeId);
+	int32 TakeId = InProcessHandle->GetTakeId();
+	FString TakeFullPath = Owner->GetFullTakePath(TakeId);
 
 	if (TakeFullPath.IsEmpty())
 	{
 		check(false);
-		FString Message = FString::Printf(TEXT("Failed to lookup take directory for take id %d"), InTakeId);
-		Owner->ExecuteProcessFinishedReporter(InTakeId, MakeError(FIngestCapability_Error::InvalidArgument, MoveTemp(Message)));
+		FString Message = FString::Printf(TEXT("Failed to lookup take directory for take id %d"), TakeId);
+		Owner->ExecuteProcessFinishedReporter(InProcessHandle, MakeError(FIngestCapability_Error::InvalidArgument, MoveTemp(Message)));
 		return;
 	}
 
-	TOptional<FTakeMetadata> MaybeTakeMetadata = Owner->GetTakeMetadata(InTakeId);
+	TOptional<FTakeMetadata> MaybeTakeMetadata = Owner->GetTakeMetadata(TakeId);
 	if (!MaybeTakeMetadata.IsSet())
 	{
 		check(false);
-		FString Message = FString::Printf(TEXT("Failed to lookup take metadata for take id %d (TakePath=%s)"), InTakeId, *TakeFullPath);
-		Owner->ExecuteProcessFinishedReporter(InTakeId, MakeError(FIngestCapability_Error::InvalidArgument, MoveTemp(Message)));
+		FString Message = FString::Printf(TEXT("Failed to lookup take metadata for take id %d (TakePath=%s)"), TakeId, *TakeFullPath);
+		Owner->ExecuteProcessFinishedReporter(InProcessHandle, MakeError(FIngestCapability_Error::InvalidArgument, MoveTemp(Message)));
 		return;
 	}
 
@@ -190,11 +150,11 @@ void UBaseIngestLiveLinkDevice::FImpl::IngestTake(int32 InTakeId, const UIngestC
 	Params.TakeOriginDirectory = MoveTemp(TakeFullPath);
 	Params.TakeOutputDirectory = FPaths::Combine(TakeConversionDirectory, TakeName + TEXT("_") + TakeUploadId.ToString());
 
-	FCaptureDataConverterResult<void> Result = RunConversion(InTakeId, IngestOptions, InTaskProgress, Params);
+	FCaptureDataConverterResult<void> Result = RunConversion(TakeId, IngestOptions, InTaskProgress, Params);
 
 	{
 		FScopeLock ConvertersRemoveLock(&Mutex);
-		Converters.Remove(InTakeId);
+		Converters.Remove(TakeId);
 	}
 
 	if (Result.HasValue())
@@ -204,7 +164,7 @@ void UBaseIngestLiveLinkDevice::FImpl::IngestTake(int32 InTakeId, const UIngestC
 			TryCleanUpWorkingDirectory(Params.TakeOutputDirectory);
 
 			FString Message = TEXT("The upload endpoint was not specified (it was an empty string), try setting a default in the hub settings");
-			Owner->ExecuteProcessFinishedReporter(InTakeId, MakeError(FIngestCapability_Error::UnrealEndpointNotFound, MoveTemp(Message)));
+			Owner->ExecuteProcessFinishedReporter(InProcessHandle, MakeError(FIngestCapability_Error::UnrealEndpointNotFound, MoveTemp(Message)));
 			return;
 		}
 
@@ -223,7 +183,7 @@ void UBaseIngestLiveLinkDevice::FImpl::IngestTake(int32 InTakeId, const UIngestC
 			TryCleanUpWorkingDirectory(Params.TakeOutputDirectory);
 
 			FString Message = FString::Printf(TEXT("Upload failed because the requested endpoint was not found: %s"), *IngestOptions->UploadHostName);
-			Owner->ExecuteProcessFinishedReporter(InTakeId, MakeError(FIngestCapability_Error::UnrealEndpointNotFound, MoveTemp(Message)));
+			Owner->ExecuteProcessFinishedReporter(InProcessHandle, MakeError(FIngestCapability_Error::UnrealEndpointNotFound, MoveTemp(Message)));
 			return;
 		}
 
@@ -243,7 +203,7 @@ void UBaseIngestLiveLinkDevice::FImpl::IngestTake(int32 InTakeId, const UIngestC
 
 			FString EndpointInfoString = UnrealEndpointInfoToString(UnrealEndpoint->GetInfo());
 			FString Message = FString::Printf(TEXT("Upload failed because we timed out connecting to the endpoint: %s"), *EndpointInfoString);
-			Owner->ExecuteProcessFinishedReporter(InTakeId, MakeError(FIngestCapability_Error::UnrealEndpointConnectionTimedOut, MoveTemp(Message)));
+			Owner->ExecuteProcessFinishedReporter(InProcessHandle, MakeError(FIngestCapability_Error::UnrealEndpointConnectionTimedOut, MoveTemp(Message)));
 			return;
 		}
 
@@ -265,22 +225,22 @@ void UBaseIngestLiveLinkDevice::FImpl::IngestTake(int32 InTakeId, const UIngestC
 
 		TakeUploadTask->Complete().BindLambda(
 			// We need to keep the Impl alive until the task is complete, as we need to use it
-			[Impl = SharedThis(this), TakeId = InTakeId, TakeOutputDirectory = Params.TakeOutputDirectory](const FString& InMessage, const int32 InStatusCode)
+			[Impl = SharedThis(this), ProcessHandle = InProcessHandle, TakeOutputDirectory = Params.TakeOutputDirectory](const FString& InMessage, const int32 InStatusCode)
 			{
 				TryCleanUpWorkingDirectory(TakeOutputDirectory);
 
 				FScopeLock TakeUploadsRemoveLock(&Impl->Mutex);
-				Impl->TakeUploads.Remove(TakeId);
+				Impl->TakeUploads.Remove(ProcessHandle->GetTakeId());
 				TakeUploadsRemoveLock.Unlock();
 
 				if (InStatusCode == 0)
 				{
-					Impl->Owner->ExecuteProcessFinishedReporter(TakeId, MakeValue<>());
+					Impl->Owner->ExecuteProcessFinishedReporter(ProcessHandle, MakeValue<>());
 				}
 				else
 				{
 					TValueOrError<void, FIngestCapability_Error> Err = MakeError(FIngestCapability_Error::UnrealEndpointUploadError, InMessage);
-					Impl->Owner->ExecuteProcessFinishedReporter(TakeId, MoveTemp(Err));
+					Impl->Owner->ExecuteProcessFinishedReporter(ProcessHandle, MoveTemp(Err));
 				}
 			}
 		);
@@ -296,7 +256,7 @@ void UBaseIngestLiveLinkDevice::FImpl::IngestTake(int32 InTakeId, const UIngestC
 			TakeUploadInfo.EndpointID = UnrealEndpoint->GetInfo().EndpointID;
 
 			FScopeLock TakeUploadsEmplaceLock(&Mutex);
-			TakeUploads.Emplace(InTakeId, MoveTemp(TakeUploadInfo));
+			TakeUploads.Emplace(InProcessHandle->GetTakeId(), MoveTemp(TakeUploadInfo));
 		}
 		else
 		{
@@ -311,7 +271,7 @@ void UBaseIngestLiveLinkDevice::FImpl::IngestTake(int32 InTakeId, const UIngestC
 		FCaptureDataConverterError Error = Result.StealError();
 
 		FText Message = FText::Join(FText::FromString(TEXT("\n\r")), Error.GetErrors());
-		Owner->ExecuteProcessFinishedReporter(InTakeId, MakeError(FIngestCapability_Error::ConversionError, Message.ToString()));
+		Owner->ExecuteProcessFinishedReporter(InProcessHandle, MakeError(FIngestCapability_Error::ConversionError, Message.ToString()));
 	}
 }
 
@@ -410,12 +370,26 @@ void UBaseIngestLiveLinkDevice::OnDeviceRemoved()
 	Super::OnDeviceRemoved();
 }
 
-void UBaseIngestLiveLinkDevice::IngestTake(int32 InTakeId, const UIngestCapability_Options* InIngestOptions, TSharedPtr<UE::CaptureManager::FTaskProgress> InTaskProgress)
+void UBaseIngestLiveLinkDevice::IngestTake(const UIngestCapability_ProcessHandle* InProcessHandle, 
+										   const UIngestCapability_Options* InIngestOptions, 
+										   TSharedPtr<UE::CaptureManager::FTaskProgress> InTaskProgress)
+{
+	Impl->IngestTake(InProcessHandle, InIngestOptions, MoveTemp(InTaskProgress));
+}
+
+void UBaseIngestLiveLinkDevice::RunDownloadTake(const UIngestCapability_ProcessHandle* InProcessHandle, const UIngestCapability_Options* InIngestOptions)
 {
-	Impl->IngestTake(InTakeId, InIngestOptions, MoveTemp(InTaskProgress));
+	// No download needed so this step is always a success
+	ExecuteProcessProgressReporter(InProcessHandle, 1.f);
+	ExecuteProcessFinishedReporter(InProcessHandle, MakeValue());
 }
 
 void UBaseIngestLiveLinkDevice::CancelIngest(int32 InTakeId)
 {
 	Impl->CancelIngest(InTakeId);
+}
+
+void UBaseIngestLiveLinkDevice::CancelIngestProcess_Implementation(const UIngestCapability_ProcessHandle* InProcessHandle)
+{
+	Impl->CancelIngest(InProcessHandle->GetTakeId());
 }
\ No newline at end of file
diff --git a/Engine/Plugins/VirtualProduction/CaptureManager/CaptureManagerApp/Source/IngestLiveLinkDevice/Public/BaseIngestLiveLinkDevice.h b/Engine/Plugins/VirtualProduction/CaptureManager/CaptureManagerApp/Source/IngestLiveLinkDevice/Public/BaseIngestLiveLinkDevice.h
--- a/Engine/Plugins/VirtualProduction/CaptureManager/CaptureManagerApp/Source/IngestLiveLinkDevice/Public/BaseIngestLiveLinkDevice.h
+++ b/Engine/Plugins/VirtualProduction/CaptureManager/CaptureManagerApp/Source/IngestLiveLinkDevice/Public/BaseIngestLiveLinkDevice.h
@@ -42,11 +42,11 @@ protected:
 
 	/**
 	 * Function implements the default behaviour for conversion and upload to the UE of the take data
-	 * @params TakeId Take identifier
+	 * @params ProcessHandle Pointer to the process handle created for this process
 	 * @params IngestOptions Ingest options used for conversion process
 	 * @params TaskProgress Helper class that accumulates progress and reports the progress to the user
 	 */
-	void IngestTake(int32 InTakeId, const UIngestCapability_Options* InIngestOptions, TSharedPtr<UE::CaptureManager::FTaskProgress> InTaskProgress);
+	void IngestTake(const UIngestCapability_ProcessHandle* InProcessHandle, const UIngestCapability_Options* InIngestOptions, TSharedPtr<UE::CaptureManager::FTaskProgress> InTaskProgress);
 
 	/**
 	 * Pure virtual function that returns the full path to the take data. Full path to the take data is required for the conversion step.
@@ -57,24 +57,36 @@ protected:
 
 	//~ Begin ULiveLinkDeviceCapability_Ingest interface
 
+	/**
+	 * Inherited function from the ILiveLinkDeviceCapability_Ingest. The function is required to be implemented.
+	 * @note In it's implementation, user may choose to call IngestTake function or to implement custom download step
+	 * @params ProcessHandle Pointer to the process handle created for this process
+	 * @params IngestOptions Ingest options used for conversion process
+	 */
+	virtual void RunDownloadTake(const UIngestCapability_ProcessHandle* InProcessHandle, const UIngestCapability_Options* InIngestOptions) override;
+
 	/**
 	 * Inherited function from the ILiveLinkDeviceCapability_Ingest. The function is required to be implemented.
 	 * @note In it's implementation, user may choose to call IngestTake function or to implement custom conversion and upload step
-	 * @params TakeId Take identifier
+	 * @params ProcessHandle Pointer to the process handle created for this process
 	 * @params IngestOptions Ingest options used for conversion process
 	 */
-	virtual void RunIngestTake(int32 InTakeId, const UIngestCapability_Options* InIngestOptions) override PURE_VIRTUAL(UBaseIngestLiveLinkDevice::RunIngestTake);
+	virtual void RunConvertAndUploadTake(const UIngestCapability_ProcessHandle* InProcessHandle, const UIngestCapability_Options* InIngestOptions) override PURE_VIRTUAL(UBaseIngestLiveLinkDevice::RunConvertAndUploadTake);
+
+	//~ End ULiveLinkDeviceCapability_Ingest interface
 
 	/**
 	 * Inherited function from the ILiveLinkDeviceCapability_Ingest. The function is required to be implemented.
 	 * @note In it's default implementation, it will abort the default ingest process.
 	 * @params TakeId Take identifier
 	 */
-	virtual void CancelIngest(int32 InTakeId) override;
-	//~ End ULiveLinkDeviceCapability_Ingest interface
+	void CancelIngest(int32 InTakeId);
+
 
 private:
 
+	virtual void CancelIngestProcess_Implementation(const UIngestCapability_ProcessHandle* InProcessHandle) override;
+
 	class FImpl;
 	TSharedPtr<FImpl> Impl;
 };
diff --git a/Engine/Plugins/VirtualProduction/CaptureManager/CaptureManagerApp/Source/LiveLinkCapabilities/Private/Ingest/IngestCapability_ProcessHandle.cpp b/Engine/Plugins/VirtualProduction/CaptureManager/CaptureManagerApp/Source/LiveLinkCapabilities/Private/Ingest/IngestCapability_ProcessHandle.cpp
--- a/Engine/Plugins/VirtualProduction/CaptureManager/CaptureManagerApp/Source/LiveLinkCapabilities/Private/Ingest/IngestCapability_ProcessHandle.cpp
+++ b/Engine/Plugins/VirtualProduction/CaptureManager/CaptureManagerApp/Source/LiveLinkCapabilities/Private/Ingest/IngestCapability_ProcessHandle.cpp
@@ -4,6 +4,24 @@
 
 #include "Ingest/LiveLinkDeviceCapability_Ingest.h"
 
+namespace UE::CaptureManager
+{
+static int32 GetNumberOfProcessTasks(EIngestCapability_ProcessConfig InProcessConfig)
+{
+	switch (InProcessConfig)
+	{
+		case EIngestCapability_ProcessConfig::Ingest:
+			return 2; // Download + Convert&Upload
+		case EIngestCapability_ProcessConfig::Download:
+			return 1;
+		case EIngestCapability_ProcessConfig::Empty:
+		default:
+			check(false);
+			return 0;
+	}
+}
+}
+
 FIngestCapability_Error::FIngestCapability_Error(ECode InCode, FString InMessage)
 	: Code(InCode)
 	, Message(MoveTemp(InMessage))
@@ -20,29 +38,73 @@ const FString& FIngestCapability_Error::GetMessage() const
 	return Message;
 }
 
+FIngestCapability_ProcessContext::FIngestCapability_ProcessContext(int32 InTakeId,
+																   EIngestCapability_ProcessConfig InProcessConfig,
+																   ILiveLinkDeviceCapability_Ingest* InOwner,
+																   FPrivateToken)
+{
+	TakeId = InTakeId;
+	ProcessConfig = InProcessConfig;
+	Owner = InOwner;
+
+	NumberOfSteps = UE::CaptureManager::GetNumberOfProcessTasks(ProcessConfig);
+}
+
+bool FIngestCapability_ProcessContext::IsDone() const
+{
+	return ProcessConfig == EIngestCapability_ProcessConfig::Empty; // ProcessConfig == 0
+}
+
 UIngestCapability_ProcessHandle::UIngestCapability_ProcessHandle()
-	: TakeId(INDEX_NONE)
-	, Owner(nullptr)
+	: Context(nullptr)
 {
 }
 
-void UIngestCapability_ProcessHandle::Initialize(int32 InTakeId, ILiveLinkDeviceCapability_Ingest* InOwner)
+void UIngestCapability_ProcessHandle::Initialize(TUniquePtr<FIngestCapability_ProcessContext> InContext)
 {
-	TakeId = InTakeId;
-	Owner = InOwner;
+	Context = MoveTemp(InContext);
 }
 
-void UIngestCapability_ProcessHandle::Run(const UIngestCapability_Options* InIngestOptions)
+int32 UIngestCapability_ProcessHandle::GetTakeId() const
 {
-	if (ensure(InIngestOptions))
-	{
-		Owner->RunIngestTake(TakeId, InIngestOptions);
-	}
+	check(Context);
+
+	return Context->TakeId;
+}
+
+bool UIngestCapability_ProcessHandle::IsDone() const
+{
+	check(Context);
+
+	return Context->IsDone();
+}
+
+FProcessFinishReporter& UIngestCapability_ProcessHandle::OnProcessFinishReporterDynamic()
+{
+	check(Context);
+
+	return Context->ProcessFinishedReporterDynamic;
 }
 
-void UIngestCapability_ProcessHandle::Cancel()
+FIngestProcessFinishReporter& UIngestCapability_ProcessHandle::OnProcessFinishReporter()
 {
-	Owner->CancelIngest(TakeId);
+	check(Context);
+
+	return Context->ProcessFinishedReporter;
+}
+
+FProcessProgressReporter& UIngestCapability_ProcessHandle::OnProcessProgressReporterDynamic()
+{
+	check(Context);
+
+	return Context->ProcessProgressReporterDynamic;
+}
+
+FIngestProcessProgressReporter& UIngestCapability_ProcessHandle::OnProcessProgressReporter()
+{
+	check(Context);
+
+	return Context->ProcessProgressReporter;
 }
 
 bool UIngestCapability_ProcessResult::IsValid() const
diff --git a/Engine/Plugins/VirtualProduction/CaptureManager/CaptureManagerApp/Source/LiveLinkCapabilities/Private/Ingest/LiveLinkDeviceCapability_Ingest.cpp b/Engine/Plugins/VirtualProduction/CaptureManager/CaptureManagerApp/Source/LiveLinkCapabilities/Private/Ingest/LiveLinkDeviceCapability_Ingest.cpp
--- a/Engine/Plugins/VirtualProduction/CaptureManager/CaptureManagerApp/Source/LiveLinkCapabilities/Private/Ingest/LiveLinkDeviceCapability_Ingest.cpp
+++ b/Engine/Plugins/VirtualProduction/CaptureManager/CaptureManagerApp/Source/LiveLinkCapabilities/Private/Ingest/LiveLinkDeviceCapability_Ingest.cpp
@@ -45,16 +45,41 @@ ILiveLinkDeviceCapability_Ingest::ILiveLinkDeviceCapability_Ingest()
 	RegisterEvent(FString(FIngestCapability_TakeRemovedEvent::Name));
 }
 
-UIngestCapability_ProcessHandle* ILiveLinkDeviceCapability_Ingest::PrepareIngestProcess_Implementation(int32 InTakeId)
+UIngestCapability_ProcessHandle* ILiveLinkDeviceCapability_Ingest::CreateIngestProcess_Implementation(int32 InTakeId, EIngestCapability_ProcessConfig InProcessConfig)
 {
-	TStrongObjectPtr<UIngestCapability_ProcessHandle> IngestProcess = CreateIngestProcessHandle(InTakeId, this);
+	TUniquePtr<FIngestCapability_ProcessContext> Context = 
+		MakeUnique<FIngestCapability_ProcessContext>(InTakeId, InProcessConfig, this, FIngestCapability_ProcessContext::FPrivateToken());
 
-	UIngestCapability_ProcessHandle* ReturnValue = IngestProcess.Get();
+	TStrongObjectPtr<UIngestCapability_ProcessHandle> IngestProcessHandle(NewObject<UIngestCapability_ProcessHandle>());
+	IngestProcessHandle->Initialize(MoveTemp(Context));
+
+	return IngestProcessHandle.Get();
+}
+
+void ILiveLinkDeviceCapability_Ingest::RunIngestProcess_Implementation(UIngestCapability_ProcessHandle* InProcessHandle, const UIngestCapability_Options* InOptions)
+{
+	using namespace UE::CaptureManager;
 
-	FScopeLock Lock(&IngestProcessMutex);
-	IngestProcessHandles.Add(InTakeId, MoveTemp(IngestProcess));
+	check(InProcessHandle->Context);
 
-	return ReturnValue;
+	FIngestCapability_ProcessContext* Context = InProcessHandle->Context.Get();
+
+	if (!Context->IngestOptions)
+	{
+		Context->IngestOptions = TStrongObjectPtr<const UIngestCapability_Options>(InOptions);
+	}
+
+	if (!Context->TaskProgress)
+	{
+		Context->TaskProgress =
+			MakeShared<FTaskProgress>(Context->NumberOfSteps,
+									  FTaskProgress::FProgressReporter::CreateLambda([this, InProcessHandle](double InProgress)
+									  {
+										  ExecuteProcessTotalProgressReporter(InProcessHandle, InProgress);
+									  }));
+	}
+
+	RunIngestProcess(InProcessHandle);
 }
 
 UIngestCapability_TakeInformation* ILiveLinkDeviceCapability_Ingest::GetTakeInformation_Implementation(int32 InTakeId) const
@@ -92,50 +117,56 @@ TArray<int32> ILiveLinkDeviceCapability_Ingest::GetTakeIdentifiers_Implementatio
 	return TakeIdentifiers;
 }
 
-void ILiveLinkDeviceCapability_Ingest::ExecuteProcessFinishedReporter(int32 InTakeIdentifier,
-																	  TValueOrError<void, FIngestCapability_Error> InMaybeError)
+void ILiveLinkDeviceCapability_Ingest::ExecuteProcessFinishedReporter(const UIngestCapability_ProcessHandle* InProcessHandle, TValueOrError<void, FIngestCapability_Error> InMaybeError)
 {
-	using namespace UE::CaptureManager;
-
-	FScopeLock Lock(&IngestProcessMutex);
-	TStrongObjectPtr<UIngestCapability_ProcessHandle>* IngestProcessHandlePtr = IngestProcessHandles.Find(InTakeIdentifier);
-	if (!IngestProcessHandlePtr)
-	{
-		return;
-	}
-
-	TStrongObjectPtr<UIngestCapability_ProcessHandle> IngestProcessHandle = *IngestProcessHandlePtr;
-	IngestProcessHandles.Remove(InTakeIdentifier);
-	Lock.Unlock();
-
 	UIngestCapability_ProcessResult* Result =
 		UIngestCapability_ProcessResult::Success();
 
+	check(InProcessHandle->Context);
+
+	FIngestCapability_ProcessContext* Context = InProcessHandle->Context.Get();
+
+	bool bShouldContinue = true;
 	if (InMaybeError.HasError())
 	{
 		FIngestCapability_Error Error = InMaybeError.GetError();
 		Result->Code = static_cast<int32>(Error.GetCode());
 		Result->Message = FText::FromString(Error.GetMessage());
+
+		bShouldContinue = false;
+	}
+	else
+	{
+		Context->ProcessConfig = static_cast<EIngestCapability_ProcessConfig>(Context->ProcessConfig & ~Context->CurrentStep);
+		bShouldContinue = !Context->IsDone();
 	}
 
-	IngestProcessHandle->ProcessFinishedReporterDynamic.ExecuteIfBound(InTakeIdentifier, Result);
-	IngestProcessHandle->ProcessFinishedReporter(InTakeIdentifier, MoveTemp(InMaybeError));	
-}
+	Context->ProcessFinishedReporterDynamic.ExecuteIfBound(InProcessHandle, Result);
+	Context->ProcessFinishedReporter(InProcessHandle, MoveTemp(InMaybeError));
 
-void ILiveLinkDeviceCapability_Ingest::ExecuteProcessProgressReporter(int32 InTakeIdentifier, double InProgress)
-{
-	FScopeLock Lock(&IngestProcessMutex);
-	TStrongObjectPtr<UIngestCapability_ProcessHandle>* IngestProcessHandlePtr = IngestProcessHandles.Find(InTakeIdentifier);
-	if (!IngestProcessHandlePtr)
+	if (bShouldContinue)
 	{
-		return;
+		RunIngestProcess(InProcessHandle);
 	}
+}
+
+void ILiveLinkDeviceCapability_Ingest::ExecuteProcessProgressReporter(const UIngestCapability_ProcessHandle* InProcessHandle, double InProgress)
+{
+	check(InProcessHandle->Context);
+
+	FIngestCapability_ProcessContext* Context = InProcessHandle->Context.Get();
+
+	Context->CurrentTask.Update(InProgress);
+}
+
+void ILiveLinkDeviceCapability_Ingest::ExecuteProcessTotalProgressReporter(const UIngestCapability_ProcessHandle* InProcessHandle, double InProgress)
+{
+	check(InProcessHandle->Context);
 
-	TStrongObjectPtr<UIngestCapability_ProcessHandle> IngestProcessHandle = *IngestProcessHandlePtr;
-	Lock.Unlock();
+	FIngestCapability_ProcessContext* Context = InProcessHandle->Context.Get();
 
-	IngestProcessHandle->ProcessProgressReporterDynamic.ExecuteIfBound(InTakeIdentifier, InProgress);
-	IngestProcessHandle->ProcessProgressReporter(InTakeIdentifier, InProgress);
+	Context->ProcessProgressReporter(InProcessHandle, InProgress);
+	Context->ProcessProgressReporterDynamic.ExecuteIfBound(InProcessHandle, InProgress);
 }
 
 void ILiveLinkDeviceCapability_Ingest::ExecuteUpdateTakeListCallback(UIngestCapability_UpdateTakeListCallback* InCallback, const TArray<int32>& InTakeIdentifiers)
@@ -192,11 +223,22 @@ TOptional<FTakeMetadata> ILiveLinkDeviceCapability_Ingest::GetTakeMetadata(int32
 	return {};
 }
 
-TStrongObjectPtr<UIngestCapability_ProcessHandle> 
-ILiveLinkDeviceCapability_Ingest::CreateIngestProcessHandle(int32 InTakeId, ILiveLinkDeviceCapability_Ingest* InOwner)
+void ILiveLinkDeviceCapability_Ingest::RunIngestProcess(const UIngestCapability_ProcessHandle* InProcessHandle)
 {
-	TStrongObjectPtr<UIngestCapability_ProcessHandle> IngestProcessHandle(NewObject<UIngestCapability_ProcessHandle>());
-	IngestProcessHandle->Initialize(InTakeId, InOwner);
+	check(InProcessHandle->Context);
+
+	FIngestCapability_ProcessContext* Context = InProcessHandle->Context.Get();
 
-	return IngestProcessHandle;
+	Context->CurrentTask = Context->TaskProgress->StartTask();
+
+	if (Context->ProcessConfig & EIngestCapability_ProcessStep::DownloadStep)
+	{
+		Context->CurrentStep = EIngestCapability_ProcessStep::DownloadStep;
+		RunDownloadTake(InProcessHandle, Context->IngestOptions.Get());
+	}
+	else if (Context->ProcessConfig & EIngestCapability_ProcessStep::ConvertAndUploadStep)
+	{
+		Context->CurrentStep = EIngestCapability_ProcessStep::ConvertAndUploadStep;
+		RunConvertAndUploadTake(InProcessHandle, Context->IngestOptions.Get());
+	}
 }
\ No newline at end of file
diff --git a/Engine/Plugins/VirtualProduction/CaptureManager/CaptureManagerApp/Source/LiveLinkCapabilities/Public/Ingest/IngestCapability_ProcessHandle.h b/Engine/Plugins/VirtualProduction/CaptureManager/CaptureManagerApp/Source/LiveLinkCapabilities/Public/Ingest/IngestCapability_ProcessHandle.h
--- a/Engine/Plugins/VirtualProduction/CaptureManager/CaptureManagerApp/Source/LiveLinkCapabilities/Public/Ingest/IngestCapability_ProcessHandle.h
+++ b/Engine/Plugins/VirtualProduction/CaptureManager/CaptureManagerApp/Source/LiveLinkCapabilities/Public/Ingest/IngestCapability_ProcessHandle.h
@@ -10,6 +10,8 @@
 #include "IngestCapability_Options.h"
 #include "IngestCapability_TakeInformation.h"
 
+#include "Async/TaskProgress.h"
+
 #include "IngestCapability_ProcessHandle.generated.h"
 
 class LIVELINKCAPABILITIES_API FIngestCapability_Error
@@ -40,12 +42,27 @@ private:
 	FString Message;
 };
 
+enum EIngestCapability_ProcessStep : uint8
+{
+	EmptyStep = 0,
+	DownloadStep = 1<<0,
+	ConvertAndUploadStep = 1<<1
+};
+
+UENUM(BlueprintType)
+enum EIngestCapability_ProcessConfig : uint32
+{
+	Empty = EIngestCapability_ProcessStep::EmptyStep UMETA(Hidden),
+	Download = EIngestCapability_ProcessStep::DownloadStep  UMETA(DisplayName = "Download", ToolTip = "Downloads takes from a device. Produces archives which can be ingested using the Take Archive Ingest device"),
+	Ingest = EIngestCapability_ProcessStep::DownloadStep | EIngestCapability_ProcessStep::ConvertAndUploadStep UMETA(DisplayName = "Ingest", ToolTip = "Downloads, converts and uploads takes to the specified client Editor")
+};
+
 UDELEGATE()
-DECLARE_DYNAMIC_DELEGATE_TwoParams(FProcessFinishReporter, int32, TakeIdentifier, UIngestCapability_ProcessResult*, IngestProcessResult);
-DECLARE_DYNAMIC_DELEGATE_TwoParams(FProcessProgressReporter, int32, TakeIdentifier, double, Progress);
+DECLARE_DYNAMIC_DELEGATE_TwoParams(FProcessFinishReporter, const UIngestCapability_ProcessHandle*, ProcessHandle, UIngestCapability_ProcessResult*, IngestProcessResult);
+DECLARE_DYNAMIC_DELEGATE_TwoParams(FProcessProgressReporter, const UIngestCapability_ProcessHandle*, ProcessHandle, double, Progress);
 
-using FIngestProcessProgressReporter = UE::CaptureManager::TManagedDelegate<int32, double>;
-using FIngestProcessFinishReporter = UE::CaptureManager::TManagedDelegate<int32, TValueOrError<void, FIngestCapability_Error>>;
+using FIngestProcessProgressReporter = UE::CaptureManager::TManagedDelegate<const UIngestCapability_ProcessHandle*, double>;
+using FIngestProcessFinishReporter = UE::CaptureManager::TManagedDelegate<const UIngestCapability_ProcessHandle*, TValueOrError<void, FIngestCapability_Error>>;
 
 UCLASS(BlueprintType)
 class LIVELINKCAPABILITIES_API UIngestCapability_ProcessResult : public UObject
@@ -70,6 +87,43 @@ public:
 	static UIngestCapability_ProcessResult* Error(FText InMessage, int32 InCode = -1);
 };
 
+class LIVELINKCAPABILITIES_API FIngestCapability_ProcessContext
+{
+private:
+
+	struct FPrivateToken { explicit FPrivateToken() = default; };
+
+public:
+
+	explicit FIngestCapability_ProcessContext(int32 InTakeId,
+											  EIngestCapability_ProcessConfig InProcessConfig,
+											  class ILiveLinkDeviceCapability_Ingest* InOwner,
+											  FPrivateToken);
+
+	bool IsDone() const;
+
+	FProcessFinishReporter ProcessFinishedReporterDynamic;
+	FIngestProcessFinishReporter ProcessFinishedReporter;
+
+	FProcessProgressReporter ProcessProgressReporterDynamic;
+	FIngestProcessProgressReporter ProcessProgressReporter;
+
+	int32 TakeId = INDEX_NONE;
+
+private:
+
+	TStrongObjectPtr<const UIngestCapability_Options> IngestOptions;
+	TSharedPtr<UE::CaptureManager::FTaskProgress> TaskProgress;
+	UE::CaptureManager::FTaskProgress::FTask CurrentTask;
+
+	EIngestCapability_ProcessConfig ProcessConfig = EIngestCapability_ProcessConfig::Empty;
+	EIngestCapability_ProcessStep CurrentStep = EIngestCapability_ProcessStep::EmptyStep;
+	int32 NumberOfSteps = 0;
+	class ILiveLinkDeviceCapability_Ingest* Owner;
+
+	friend class ILiveLinkDeviceCapability_Ingest;
+};
+
 UCLASS(BlueprintType)
 class LIVELINKCAPABILITIES_API UIngestCapability_ProcessHandle : public UObject
 {
@@ -77,32 +131,29 @@ class LIVELINKCAPABILITIES_API UIngestCapability_ProcessHandle : public UObject
 
 public:
 
-	UPROPERTY(EditAnywhere, Category = "Live Link Device|Ingest")
-	int32 TakeId = INDEX_NONE;
-
 	UFUNCTION(BlueprintCallable, Category = "Live Link Device|Ingest")
-	void Run(const UIngestCapability_Options* InIngestOptions);
+	int32 GetTakeId() const;
 
 	UFUNCTION(BlueprintCallable, Category = "Live Link Device|Ingest")
-	void Cancel();
+	bool IsDone() const;
 
-	UPROPERTY(EditAnywhere, Category = "Live Link Device|Ingest")
-	FProcessFinishReporter ProcessFinishedReporterDynamic;
+	UFUNCTION(BlueprintCallable, Category = "Live Link Device|Ingest")
+	FProcessFinishReporter& OnProcessFinishReporterDynamic();
 
-	FIngestProcessFinishReporter ProcessFinishedReporter;
+	FIngestProcessFinishReporter& OnProcessFinishReporter();
 
-	UPROPERTY(EditAnywhere, Category = "Live Link Device|Ingest")
-	FProcessProgressReporter ProcessProgressReporterDynamic;
+	UFUNCTION(BlueprintCallable, Category = "Live Link Device|Ingest")
+	FProcessProgressReporter& OnProcessProgressReporterDynamic();
 
-	FIngestProcessProgressReporter ProcessProgressReporter;
+	FIngestProcessProgressReporter& OnProcessProgressReporter();
 
 private:
 
 	UIngestCapability_ProcessHandle();
 
-	void Initialize(int32 InTakeId, class ILiveLinkDeviceCapability_Ingest* InOwner);
+	void Initialize(TUniquePtr<FIngestCapability_ProcessContext> InContext);
 
-	class ILiveLinkDeviceCapability_Ingest* Owner;
+	TUniquePtr<FIngestCapability_ProcessContext> Context;
 
 	friend class ILiveLinkDeviceCapability_Ingest;
 };
\ No newline at end of file
diff --git a/Engine/Plugins/VirtualProduction/CaptureManager/CaptureManagerApp/Source/LiveLinkCapabilities/Public/Ingest/LiveLinkDeviceCapability_Ingest.h b/Engine/Plugins/VirtualProduction/CaptureManager/CaptureManagerApp/Source/LiveLinkCapabilities/Public/Ingest/LiveLinkDeviceCapability_Ingest.h
--- a/Engine/Plugins/VirtualProduction/CaptureManager/CaptureManagerApp/Source/LiveLinkCapabilities/Public/Ingest/LiveLinkDeviceCapability_Ingest.h
+++ b/Engine/Plugins/VirtualProduction/CaptureManager/CaptureManagerApp/Source/LiveLinkCapabilities/Public/Ingest/LiveLinkDeviceCapability_Ingest.h
@@ -46,7 +46,13 @@ public:
 	ILiveLinkDeviceCapability_Ingest();
 
 	UFUNCTION(BlueprintCallable, BlueprintNativeEvent, Category = "Live Link Device|Ingest")
-	UIngestCapability_ProcessHandle* PrepareIngestProcess(int32 InTakeId);
+	UIngestCapability_ProcessHandle* CreateIngestProcess(int32 InTakeId, EIngestCapability_ProcessConfig InProcessConfig);
+
+	UFUNCTION(BlueprintCallable, BlueprintNativeEvent, Category = "Live Link Device|Ingest")
+	void RunIngestProcess(UIngestCapability_ProcessHandle* InProcessHandle, const UIngestCapability_Options* InOptions);
+
+	UFUNCTION(BlueprintCallable, BlueprintNativeEvent, Category = "Live Link Device|Ingest")
+	void CancelIngestProcess(const UIngestCapability_ProcessHandle* InProcessHandle);
 
 	UFUNCTION(BlueprintCallable, BlueprintNativeEvent, Category = "Live Link Device|Ingest")
 	void UpdateTakeList(UIngestCapability_UpdateTakeListCallback* InCallback);
@@ -57,7 +63,10 @@ public:
 	UFUNCTION(BlueprintCallable, BlueprintNativeEvent, Category = "Live Link Device|Ingest")
 	TArray<int32> GetTakeIdentifiers() const;
 
-	UIngestCapability_ProcessHandle* PrepareIngestProcess_Implementation(int32 InTakeId);
+	UIngestCapability_ProcessHandle* CreateIngestProcess_Implementation(int32 InTakeId, EIngestCapability_ProcessConfig InProcessConfig);
+	void RunIngestProcess_Implementation(UIngestCapability_ProcessHandle* InProcessHandle, const UIngestCapability_Options* InOptions);
+	virtual void CancelIngestProcess_Implementation(const UIngestCapability_ProcessHandle* InProcessHandle) PURE_VIRTUAL(ILiveLinkDeviceCapability_Ingest::CancelIngestProcess_Implementation);
+
 	virtual void UpdateTakeList_Implementation(UIngestCapability_UpdateTakeListCallback* InCallback) PURE_VIRTUAL(ILiveLinkDeviceCapability_Ingest::UpdateTakeList_Implementation);
 	UIngestCapability_TakeInformation* GetTakeInformation_Implementation(int32 InTakeId) const;
 	TArray<int32> GetTakeIdentifiers_Implementation() const;
@@ -66,8 +75,8 @@ public:
 
 protected:
 
-	void ExecuteProcessFinishedReporter(int32 InTakeIdentifier, TValueOrError<void, FIngestCapability_Error> InMaybeError);
-	void ExecuteProcessProgressReporter(int32 InTakeIdentifier, double InProgress);
+	void ExecuteProcessFinishedReporter(const UIngestCapability_ProcessHandle* InProcessHandle, TValueOrError<void, FIngestCapability_Error> InMaybeError);
+	void ExecuteProcessProgressReporter(const UIngestCapability_ProcessHandle* InProcessHandle, double InProgress);
 	void ExecuteUpdateTakeListCallback(UIngestCapability_UpdateTakeListCallback* InCallback, const TArray<int32>& InTakeIdentifiers);
 	
 	int32 AddTake(FTakeMetadata InTakeMetadata);
@@ -77,18 +86,16 @@ protected:
 
 private:
 
-	static TStrongObjectPtr<UIngestCapability_ProcessHandle>
-		CreateIngestProcessHandle(int32 InTakeId, ILiveLinkDeviceCapability_Ingest* InOwner);
+	void ExecuteProcessTotalProgressReporter(const UIngestCapability_ProcessHandle* InProcessHandle, double InProgress);
+
+	void RunIngestProcess(const UIngestCapability_ProcessHandle* InProcessHandle);
 
-	virtual void CancelIngest(int32 InTakeId) = 0;
-	virtual void RunIngestTake(int32 InTakeId, const UIngestCapability_Options* InIngestOptions) = 0;
+	virtual void RunDownloadTake(const UIngestCapability_ProcessHandle* InProcessHandle, const UIngestCapability_Options* InIngestOptions) = 0;
+	virtual void RunConvertAndUploadTake(const UIngestCapability_ProcessHandle* InProcessHandle, const UIngestCapability_Options* InIngestOptions) = 0;
 
 	mutable FCriticalSection TakeMetadataMutex;
 	TMap<int32, FTakeMetadata> TakeMetadataMap;
 
-	mutable FCriticalSection IngestProcessMutex;
-	TMap<int32, TStrongObjectPtr<UIngestCapability_ProcessHandle>> IngestProcessHandles;
-
 	std::atomic_int32_t CurrentTakeId = 0;
 
 	friend class UIngestCapability_ProcessHandle;
diff --git a/Engine/Plugins/VirtualProduction/CaptureManager/CaptureManagerCore/Source/CaptureUtils/Private/Async/TaskProgress.cpp b/Engine/Plugins/VirtualProduction/CaptureManager/CaptureManagerCore/Source/CaptureUtils/Private/Async/TaskProgress.cpp
--- a/Engine/Plugins/VirtualProduction/CaptureManager/CaptureManagerCore/Source/CaptureUtils/Private/Async/TaskProgress.cpp
+++ b/Engine/Plugins/VirtualProduction/CaptureManager/CaptureManagerCore/Source/CaptureUtils/Private/Async/TaskProgress.cpp
@@ -7,16 +7,63 @@
 namespace UE::CaptureManager
 {
 
-FTaskProgress::FTask::FTask(FTaskProgress& InTaskProgress, int32 InId)
-	: TaskProgress(InTaskProgress)
+FTaskProgress::FTask::FTask()
+	: TaskProgress(nullptr)
+{
+}
+
+FTaskProgress::FTask::FTask(FTask&& InTask)
+{
+	*this = MoveTemp(InTask);
+}
+
+FTaskProgress::FTask& FTaskProgress::FTask::operator=(FTask&& InTask)
+{
+	if (this == &InTask)
+	{
+		return *this;
+	}
+
+	TaskProgress = MoveTemp(InTask.TaskProgress);
+	Id = InTask.Id;
+
+	InTask.TaskProgress = nullptr;
+	InTask.Id = INDEX_NONE;
+
+	return *this;
+}
+
+FTaskProgress::FTask::FTask(const FTask& InTask)
+{
+	*this = InTask;
+}
+
+FTaskProgress::FTask& FTaskProgress::FTask::operator=(const FTask& InTask)
+{
+	if (this == &InTask)
+	{
+		return *this;
+	}
+
+	TaskProgress = InTask.TaskProgress;
+	Id = InTask.Id;
+
+	return *this;
+}
+
+FTaskProgress::FTask::FTask(TWeakPtr<FTaskProgress> InTaskProgress, int32 InId)
+	: TaskProgress(MoveTemp(InTaskProgress))
 	, Id(InId)
 {
 }
 
 void FTaskProgress::FTask::Update(double InProgress)
 {
-	TaskProgress.Update(Id, InProgress);
-	TaskProgress.Report();
+	if (TSharedPtr<FTaskProgress> TaskProgressShared = TaskProgress.Pin())
+	{
+		TaskProgressShared->Update(Id, InProgress);
+		TaskProgressShared->Report();
+	}
 }
 
 FTaskProgress::FTaskProgress(uint32 InAmountOfWork, FProgressReporter InReporter)
@@ -32,7 +79,7 @@ FTaskProgress::FTask FTaskProgress::StartTask()
 	int32 NewTask = ++CurrentTask;
 	checkf(CurrentProgressValues.Lock()->IsValidIndex(NewTask), TEXT("Unexpected start task"));
 
-	return FTaskProgress::FTask(*this, NewTask);
+	return FTaskProgress::FTask(AsWeak(), NewTask);
 }
 
 double FTaskProgress::GetTotalProgress() const
diff --git a/Engine/Plugins/VirtualProduction/CaptureManager/CaptureManagerCore/Source/CaptureUtils/Private/Async/TaskWaiter.cpp b/Engine/Plugins/VirtualProduction/CaptureManager/CaptureManagerCore/Source/CaptureUtils/Private/Async/TaskWaiter.cpp
new file mode 100644
--- /dev/null
+++ b/Engine/Plugins/VirtualProduction/CaptureManager/CaptureManagerCore/Source/CaptureUtils/Private/Async/TaskWaiter.cpp
@@ -0,0 +1,48 @@
+// Copyright Epic Games, Inc. All Rights Reserved.
+
+#include "Async/TaskWaiter.h"
+
+#include "HAL/PlatformProcess.h"
+
+namespace UE::CaptureManager
+{
+
+FTaskWaiter::FTaskWaiter() : TaskCounter(0)
+{
+
+}
+
+FTaskWaiter::~FTaskWaiter()
+{
+	check(TaskCounter.load() <= CanCreateTaskFlag);
+}
+
+bool FTaskWaiter::CreateTask()
+{
+	uint32 OldCounter = TaskCounter.fetch_add(1);
+
+	if ((OldCounter & CanCreateTaskFlag) > 0)
+	{
+		TaskCounter.fetch_sub(1);
+		return false;
+	}
+
+	return true;
+}
+
+void FTaskWaiter::FinishTask()
+{
+	TaskCounter.fetch_sub(1);
+}
+
+void FTaskWaiter::WaitForAll()
+{
+	TaskCounter.fetch_add(CanCreateTaskFlag);
+
+	while (TaskCounter > CanCreateTaskFlag)
+	{
+		FPlatformProcess::Sleep(0.1f);
+	}
+}
+
+}
\ No newline at end of file
diff --git a/Engine/Plugins/VirtualProduction/CaptureManager/CaptureManagerCore/Source/CaptureUtils/Public/Async/TaskProgress.h b/Engine/Plugins/VirtualProduction/CaptureManager/CaptureManagerCore/Source/CaptureUtils/Public/Async/TaskProgress.h
--- a/Engine/Plugins/VirtualProduction/CaptureManager/CaptureManagerCore/Source/CaptureUtils/Public/Async/TaskProgress.h
+++ b/Engine/Plugins/VirtualProduction/CaptureManager/CaptureManagerCore/Source/CaptureUtils/Public/Async/TaskProgress.h
@@ -11,7 +11,7 @@
 namespace UE::CaptureManager
 {
 
-class CAPTUREUTILS_API FTaskProgress
+class CAPTUREUTILS_API FTaskProgress : public TSharedFromThis<FTaskProgress>
 {
 public:
 
@@ -21,13 +21,21 @@ public:
 	{
 	public:
 
+		FTask();
+
+		FTask(const FTask& InTask);
+		FTask& operator=(const FTask& InTask);
+
+		FTask(FTask&& InTask);
+		FTask& operator=(FTask&& InTask);
+
 		void Update(double InProgress);
 
 	private:
 
-		FTask(FTaskProgress& InTaskProgress, int32 InId);
+		FTask(TWeakPtr<FTaskProgress> InTaskProgress, int32 InId);
 
-		FTaskProgress& TaskProgress;
+		TWeakPtr<FTaskProgress> TaskProgress;
 		int32 Id;
 
 		friend FTaskProgress;
diff --git a/Engine/Plugins/VirtualProduction/CaptureManager/CaptureManagerCore/Source/CaptureUtils/Public/Async/TaskWaiter.h b/Engine/Plugins/VirtualProduction/CaptureManager/CaptureManagerCore/Source/CaptureUtils/Public/Async/TaskWaiter.h
new file mode 100644
--- /dev/null
+++ b/Engine/Plugins/VirtualProduction/CaptureManager/CaptureManagerCore/Source/CaptureUtils/Public/Async/TaskWaiter.h
@@ -0,0 +1,26 @@
+// Copyright Epic Games, Inc. All Rights Reserved.
+
+#include <atomic>
+#include "HAL/Platform.h"
+
+namespace UE::CaptureManager
+{
+
+class CAPTUREUTILS_API FTaskWaiter
+{
+public:
+
+	FTaskWaiter();
+	~FTaskWaiter();
+
+	bool CreateTask();
+	void FinishTask();
+	void WaitForAll();
+
+private:
+
+	std::atomic<uint32> TaskCounter;
+	const uint32 CanCreateTaskFlag = 0x80000000;
+};
+
+}
\ No newline at end of file
diff --git a/Engine/Plugins/VirtualProduction/CaptureManager/CaptureManagerDevices/Source/CPSLiveLinkDevice/Private/LiveLinkFaceDevice.cpp b/Engine/Plugins/VirtualProduction/CaptureManager/CaptureManagerDevices/Source/CPSLiveLinkDevice/Private/LiveLinkFaceDevice.cpp
--- a/Engine/Plugins/VirtualProduction/CaptureManager/CaptureManagerDevices/Source/CPSLiveLinkDevice/Private/LiveLinkFaceDevice.cpp
+++ b/Engine/Plugins/VirtualProduction/CaptureManager/CaptureManagerDevices/Source/CPSLiveLinkDevice/Private/LiveLinkFaceDevice.cpp
@@ -94,19 +94,13 @@ void ULiveLinkFaceDevice::UpdateTakeList_Implementation(UIngestCapability_Update
 	ExecuteUpdateTakeListCallback(InCallback, Execute_GetTakeIdentifiers(this));
 }
 
-void ULiveLinkFaceDevice::RunIngestTake(int32 InTakeId, const UIngestCapability_Options* InIngestOptions)
+void ULiveLinkFaceDevice::RunDownloadTake(const UIngestCapability_ProcessHandle* InProcessHandle, const UIngestCapability_Options* InIngestOptions)
 {
 	using namespace UE::CaptureManager;
 
-	const FGetTakeMetadataResponse::FTakeObject& Take = Device->GetTake(InTakeId);
+	FTakeId TakeId = InProcessHandle->GetTakeId();
 
-	static constexpr uint32 NumberOfTasks = 3; // Download, Convert, Upload
-
-	TSharedPtr<FTaskProgress> TaskProgress = MakeShared<FTaskProgress>
-		(NumberOfTasks, FTaskProgress::FProgressReporter::CreateLambda([this, TakeId = InTakeId](double InProgress)
-	{
-		ExecuteProcessProgressReporter(TakeId, InProgress);
-	}));
+	const FGetTakeMetadataResponse::FTakeObject& Take = Device->GetTake(TakeId);
 
 	uint64 TotalSize = 0;
 	for (const FGetTakeMetadataResponse::FFileObject& File : Take.Files)
@@ -116,22 +110,54 @@ void ULiveLinkFaceDevice::RunIngestTake(int32 InTakeId, const UIngestCapability_
 
 	FString DownloadedStorage = InIngestOptions->DownloadDirectory;
 
-	FTaskProgress::FTask Task = TaskProgress->StartTask();
-
 	TStrongObjectPtr<const UIngestCapability_Options> IngestOptions(InIngestOptions);
+	TStrongObjectPtr<const UIngestCapability_ProcessHandle> ProcessHandle(InProcessHandle);
 
 	TUniquePtr<FCPSFileStream> CPSStream = MakeUnique<FCPSFileStream>(MoveTemp(DownloadedStorage), TotalSize);
-	CPSStream->SetExportFinished(FCPSFileStream::FExportFinished::CreateUObject(this, &ULiveLinkFaceDevice::OnExportFinished, Take.Name, InTakeId, TaskProgress, IngestOptions));
-	CPSStream->SetProgressHandler(FCPSFileStream::FReportProgress::CreateUObject(this, &ULiveLinkFaceDevice::OnExportProgressReport, Task));
+	CPSStream->SetExportFinished(FCPSFileStream::FExportFinished::CreateUObject(this, &ULiveLinkFaceDevice::OnExportFinished, Take.Name, ProcessHandle, IngestOptions));
+	CPSStream->SetProgressHandler(FCPSFileStream::FReportProgress::CreateUObject(this, &ULiveLinkFaceDevice::OnExportProgressReport, ProcessHandle));
+
+	Device->StartExport(TakeId, MoveTemp(CPSStream));
+}
+
+void ULiveLinkFaceDevice::RunConvertAndUploadTake(const UIngestCapability_ProcessHandle* InProcessHandle, const UIngestCapability_Options* InIngestOptions)
+{
+	using namespace UE::CaptureManager;
+
+	static constexpr uint32 NumberOfTasks = 2; // Convert, Upload
+
+	TStrongObjectPtr<const UIngestCapability_ProcessHandle> ProcessHandle(InProcessHandle);
+	TStrongObjectPtr<const UIngestCapability_Options> IngestOptions(InIngestOptions);
+
+	TSharedPtr<FTaskProgress> TaskProgress = MakeShared<FTaskProgress>
+		(NumberOfTasks, FTaskProgress::FProgressReporter::CreateLambda([this, ProcessHandle](double InProgress)
+		 {
+			 ExecuteProcessProgressReporter(ProcessHandle.Get(), InProgress);
+		 }));
+
+	// Free the current thread that is waiting on next download
+	TWeakObjectPtr<ULiveLinkFaceDevice> WeakThis(this);
+	AsyncTask(ENamedThreads::Type::AnyThread, [WeakThis = MoveTemp(WeakThis), ProcessHandle, IngestOptions, TaskProgress = MoveTemp(TaskProgress)]()
+	{
+		TStrongObjectPtr<ULiveLinkFaceDevice> StrongThis = WeakThis.Pin();
 
-	Device->StartExport(InTakeId, MoveTemp(CPSStream));
+		if (!StrongThis.IsValid())
+		{
+			return;
+		}
+
+		StrongThis->Super::IngestTake(ProcessHandle.Get(), IngestOptions.Get(), TaskProgress);
+		StrongThis->RemoveDownloadedTakeData(ProcessHandle->GetTakeId());
+	});
 }
 
-void ULiveLinkFaceDevice::CancelIngest(int32 InTakeId)
+void ULiveLinkFaceDevice::CancelIngestProcess_Implementation(const UIngestCapability_ProcessHandle* InProcessHandle)
 {
-	Device->CancelExport(InTakeId);
+	UE::CaptureManager::FTakeId TakeId = InProcessHandle->GetTakeId();
+
+	Device->CancelExport(TakeId);
 
-	Super::CancelIngest(InTakeId);
+	Super::CancelIngest(TakeId);
 }
 
 ELiveLinkDeviceConnectionStatus ULiveLinkFaceDevice::GetConnectionStatus_Implementation() const
@@ -362,16 +388,14 @@ FTakeMetadata ULiveLinkFaceDevice::ParseTakeMetadata(const UE::CaptureManager::F
 	return TakeMetadata;
 }
 
-void ULiveLinkFaceDevice::OnExportProgressReport(float InProgress,
-													   UE::CaptureManager::FTaskProgress::FTask InTask)
+void ULiveLinkFaceDevice::OnExportProgressReport(float InProgress, TStrongObjectPtr<const UIngestCapability_ProcessHandle> InProcessHandle)
 {
-	InTask.Update(static_cast<double>(InProgress));
+	ExecuteProcessProgressReporter(InProcessHandle.Get(), InProgress);
 }
 
 void ULiveLinkFaceDevice::OnExportFinished(UE::CaptureManager::TProtocolResult<void> InResult,
 										   FString InTakeName,
-										   UE::CaptureManager::FTakeId InTakeId,
-										   TSharedPtr<UE::CaptureManager::FTaskProgress> InTaskProgress,
+										   TStrongObjectPtr<const UIngestCapability_ProcessHandle> InProcessHandle,
 										   TStrongObjectPtr<const UIngestCapability_Options> InIngestOptions)
 {
 	using namespace UE::CaptureManager;
@@ -385,38 +409,27 @@ void ULiveLinkFaceDevice::OnExportFinished(UE::CaptureManager::TProtocolResult<v
 		TOptional<FTakeMetadata> TakeMetadata = ParseTake(DownloadedStorage, InTakeName);
 		if (TakeMetadata.IsSet())
 		{
+			FTakeId TakeId = InProcessHandle->GetTakeId();
+
 			FScopeLock Lock(&DownloadedTakesMutex);
-			DownloadedTakes.Add(InTakeId, DownloadedTake);
+			DownloadedTakes.Add(TakeId, DownloadedTake);
 			Lock.Unlock();
 
-			UpdateTake(InTakeId, MoveTemp(TakeMetadata.GetValue()));
-
-			// Free the current thread that is waiting on next download
-			TWeakObjectPtr<ULiveLinkFaceDevice> WeakThis(this);
-			AsyncTask(ENamedThreads::Type::AnyThread, [WeakThis = MoveTemp(WeakThis), TakeId = InTakeId, IngestOptions = MoveTemp(InIngestOptions), TaskProgress = MoveTemp(InTaskProgress)]()
-			{
-				TStrongObjectPtr<ULiveLinkFaceDevice> StrongThis = WeakThis.Pin();
-
-				if (!StrongThis.IsValid())
-				{
-					return;
-				}
+			UpdateTake(TakeId, MoveTemp(TakeMetadata.GetValue()));
 
-				StrongThis->Super::IngestTake(TakeId, IngestOptions.Get(), TaskProgress);
-				StrongThis->RemoveDownloadedTakeData(TakeId);
-			});	
+			ExecuteProcessFinishedReporter(InProcessHandle.Get(), MakeValue());
 		}
 		else
 		{
 			TValueOrError<void, FIngestCapability_Error> Result = MakeError(FIngestCapability_Error::DownloaderError, TEXT("Failed to parse the take metadata"));
-			ExecuteProcessFinishedReporter(InTakeId, MoveTemp(Result));
+			ExecuteProcessFinishedReporter(InProcessHandle.Get(), MoveTemp(Result));
 			IFileManager::Get().DeleteDirectory(*DownloadedTake, false, true);
 		}
 	}
 	else
 	{
 		TValueOrError<void, FIngestCapability_Error> Result = MakeError(FIngestCapability_Error::DownloaderError, InResult.GetError().GetMessage());
-		ExecuteProcessFinishedReporter(InTakeId, MoveTemp(Result));
+		ExecuteProcessFinishedReporter(InProcessHandle.Get(), MoveTemp(Result));
 		IFileManager::Get().DeleteDirectory(*DownloadedTake, false, true);
 	}
 }
diff --git a/Engine/Plugins/VirtualProduction/CaptureManager/CaptureManagerDevices/Source/CPSLiveLinkDevice/Private/LiveLinkFaceDevice.h b/Engine/Plugins/VirtualProduction/CaptureManager/CaptureManagerDevices/Source/CPSLiveLinkDevice/Private/LiveLinkFaceDevice.h
--- a/Engine/Plugins/VirtualProduction/CaptureManager/CaptureManagerDevices/Source/CPSLiveLinkDevice/Private/LiveLinkFaceDevice.h
+++ b/Engine/Plugins/VirtualProduction/CaptureManager/CaptureManagerDevices/Source/CPSLiveLinkDevice/Private/LiveLinkFaceDevice.h
@@ -66,8 +66,9 @@ private:
 
 	//~ Begin ULiveLinkDeviceCapability_Ingest interface
 	virtual void UpdateTakeList_Implementation(UIngestCapability_UpdateTakeListCallback* InCallback) override;
-	virtual void RunIngestTake(int32 InTakeId, const UIngestCapability_Options* InIngestOptions) override;
-	virtual void CancelIngest(int32 InTakeId) override;
+	virtual void RunDownloadTake(const UIngestCapability_ProcessHandle* InProcessHandle, const UIngestCapability_Options* InIngestOptions) override;
+	virtual void RunConvertAndUploadTake(const UIngestCapability_ProcessHandle* InProcessHandle, const UIngestCapability_Options* InIngestOptions) override;
+	virtual void CancelIngestProcess_Implementation(const UIngestCapability_ProcessHandle* InProcessHandle) override;
 	//~ End ULiveLinkDeviceCapability_Ingest interface
 
 	//~ Begin ILiveLinkDeviceCapability_Connection interface
@@ -88,12 +89,10 @@ private:
 	void HandleCPSStateUpdate(TSharedPtr<const UE::CaptureManager::FCaptureEvent> InEvent);
 	void HandleCPSEvent(TSharedPtr<const UE::CaptureManager::FCaptureEvent> InEvent);
 
-	void OnExportProgressReport(float InProgress,
-								UE::CaptureManager::FTaskProgress::FTask InTask);
+	void OnExportProgressReport(float InProgress, TStrongObjectPtr<const UIngestCapability_ProcessHandle> InProcessHandle);
 	void OnExportFinished(UE::CaptureManager::TProtocolResult<void> InResult,
 						  FString InTakeName,
-						  UE::CaptureManager::FTakeId InTakeId,
-						  TSharedPtr<UE::CaptureManager::FTaskProgress> InTaskProgress,
+						  TStrongObjectPtr<const UIngestCapability_ProcessHandle> InProcessHandle,
 						  TStrongObjectPtr<const UIngestCapability_Options> InIngestOptions);
 
 	TOptional<FTakeMetadata> ParseTake(const FString& InTakeDirectory, const FString& InTakeName);
diff --git a/Engine/Plugins/VirtualProduction/CaptureManager/CaptureManagerDevices/Source/FilesystemLiveLinkDevice/Private/FilesystemLiveLinkDevice.cpp b/Engine/Plugins/VirtualProduction/CaptureManager/CaptureManagerDevices/Source/FilesystemLiveLinkDevice/Private/FilesystemLiveLinkDevice.cpp
--- a/Engine/Plugins/VirtualProduction/CaptureManager/CaptureManagerDevices/Source/FilesystemLiveLinkDevice/Private/FilesystemLiveLinkDevice.cpp
+++ b/Engine/Plugins/VirtualProduction/CaptureManager/CaptureManagerDevices/Source/FilesystemLiveLinkDevice/Private/FilesystemLiveLinkDevice.cpp
@@ -69,22 +69,22 @@ FString UFilesystemLiveLinkDevice::GetFullTakePath(UE::CaptureManager::FTakeId I
 	return FString();
 }
 
-void UFilesystemLiveLinkDevice::RunIngestTake(int32 InTakeId, const UIngestCapability_Options* InIngestOptions)
+void UFilesystemLiveLinkDevice::RunConvertAndUploadTake(const UIngestCapability_ProcessHandle* InProcessHandle, const UIngestCapability_Options* InIngestOptions)
 {
 	AsyncTask(ENamedThreads::AnyThread, [this,
-			  TakeId = InTakeId,
-			  IngestOptions = InIngestOptions]()
+			  ProcessHandle = TStrongObjectPtr<const UIngestCapability_ProcessHandle>(InProcessHandle),
+			  IngestOptions = TStrongObjectPtr<const UIngestCapability_Options>(InIngestOptions)]()
 	{
 		static constexpr uint32 NumberOfTasks = 2; // Convert, Upload
 
 		using namespace UE::CaptureManager;
 		TSharedPtr<FTaskProgress> TaskProgress
-			= MakeShared<FTaskProgress>(NumberOfTasks, FTaskProgress::FProgressReporter::CreateLambda([this, TakeId](double InProgress)
+			= MakeShared<FTaskProgress>(NumberOfTasks, FTaskProgress::FProgressReporter::CreateLambda([this, ProcessHandle](double InProgress)
 		{
-			ExecuteProcessProgressReporter(TakeId, InProgress);
+			ExecuteProcessProgressReporter(ProcessHandle.Get(), InProgress);
 		}));
 
-		Super::IngestTake(TakeId, IngestOptions, MoveTemp(TaskProgress));
+		Super::IngestTake(ProcessHandle.Get(), IngestOptions.Get(), MoveTemp(TaskProgress));
 	});
 }
 
diff --git a/Engine/Plugins/VirtualProduction/CaptureManager/CaptureManagerDevices/Source/FilesystemLiveLinkDevice/Private/FilesystemLiveLinkDevice.h b/Engine/Plugins/VirtualProduction/CaptureManager/CaptureManagerDevices/Source/FilesystemLiveLinkDevice/Private/FilesystemLiveLinkDevice.h
--- a/Engine/Plugins/VirtualProduction/CaptureManager/CaptureManagerDevices/Source/FilesystemLiveLinkDevice/Private/FilesystemLiveLinkDevice.h
+++ b/Engine/Plugins/VirtualProduction/CaptureManager/CaptureManagerDevices/Source/FilesystemLiveLinkDevice/Private/FilesystemLiveLinkDevice.h
@@ -54,7 +54,7 @@ private:
 	virtual FString GetFullTakePath(UE::CaptureManager::FTakeId InTakeId) const override;
 
 	//~ Begin ULiveLinkDeviceCapability_Ingest interface
-	virtual void RunIngestTake(int32 InTakeId, const UIngestCapability_Options* InIngestOptions) override;
+	virtual void RunConvertAndUploadTake(const UIngestCapability_ProcessHandle* InProcessHandle, const UIngestCapability_Options* InIngestOptions) override;
 	virtual void UpdateTakeList_Implementation(UIngestCapability_UpdateTakeListCallback* InCallback) override;
 	//~ End ULiveLinkDeviceCapability_Ingest interface
 
diff --git a/Engine/Plugins/VirtualProduction/CaptureManager/CaptureManagerDevices/Source/FilesystemLiveLinkDevice/Private/Tests/GenericLiveLinkDeviceTests.cpp b/Engine/Plugins/VirtualProduction/CaptureManager/CaptureManagerDevices/Source/FilesystemLiveLinkDevice/Private/Tests/GenericLiveLinkDeviceTests.cpp
--- a/Engine/Plugins/VirtualProduction/CaptureManager/CaptureManagerDevices/Source/FilesystemLiveLinkDevice/Private/Tests/GenericLiveLinkDeviceTests.cpp
+++ b/Engine/Plugins/VirtualProduction/CaptureManager/CaptureManagerDevices/Source/FilesystemLiveLinkDevice/Private/Tests/GenericLiveLinkDeviceTests.cpp
@@ -28,11 +28,11 @@ bool FGenericLiveLinkArchiveDeviceIngestTakeAsyncTest::RunTest(const FString& Pa
 	TestTrue(TEXT("Device Implements Ingest Capability"), Device->Implements<ULiveLinkDeviceCapability_Ingest>());
 
 	UIngestCapability_ProcessHandle* IngestProcessHandle =
-		ILiveLinkDeviceCapability_Ingest::Execute_PrepareIngestProcess(Device, 0);
+		ILiveLinkDeviceCapability_Ingest::Execute_CreateIngestProcess(Device, 0, EIngestCapability_ProcessConfig::Ingest);
 	TestNotNull(TEXT("Ingest Process Handle Not Null"), IngestProcessHandle);
 
 	static constexpr UIngestCapability_Options* Options = nullptr;
-	IngestProcessHandle->Run(Options);
+	ILiveLinkDeviceCapability_Ingest::Execute_RunIngestProcess(Device, IngestProcessHandle, Options);
 
 	return true;
 }
\ No newline at end of file
diff --git a/Engine/Plugins/VirtualProduction/CaptureManager/CaptureManagerDevices/Source/MonoVideoLiveLinkDevice/Private/MonoVideoLiveLinkDevice.cpp b/Engine/Plugins/VirtualProduction/CaptureManager/CaptureManagerDevices/Source/MonoVideoLiveLinkDevice/Private/MonoVideoLiveLinkDevice.cpp
--- a/Engine/Plugins/VirtualProduction/CaptureManager/CaptureManagerDevices/Source/MonoVideoLiveLinkDevice/Private/MonoVideoLiveLinkDevice.cpp
+++ b/Engine/Plugins/VirtualProduction/CaptureManager/CaptureManagerDevices/Source/MonoVideoLiveLinkDevice/Private/MonoVideoLiveLinkDevice.cpp
@@ -98,22 +98,22 @@ FString UMonoVideoLiveLinkDevice::GetFullTakePath(UE::CaptureManager::FTakeId In
 	return FString();
 }
 
-void UMonoVideoLiveLinkDevice::RunIngestTake(int32 InTakeId, const UIngestCapability_Options* InIngestOptions)
+void UMonoVideoLiveLinkDevice::RunConvertAndUploadTake(const UIngestCapability_ProcessHandle* InProcessHandle, const UIngestCapability_Options* InIngestOptions)
 {
 	AsyncTask(ENamedThreads::AnyThread, [this,
-			  TakeId = InTakeId,
-			  IngestOptions = InIngestOptions]()
+			  ProcessHandle = TStrongObjectPtr<const UIngestCapability_ProcessHandle>(InProcessHandle),
+			  IngestOptions = TStrongObjectPtr<const UIngestCapability_Options>(InIngestOptions)]()
 	{
 		static constexpr uint32 NumberOfTasks = 2; // Convert, Upload
 
 		using namespace UE::CaptureManager;
 		TSharedPtr<FTaskProgress> TaskProgress
-			= MakeShared<FTaskProgress>(NumberOfTasks, FTaskProgress::FProgressReporter::CreateLambda([this, TakeId](double InProgress)
+			= MakeShared<FTaskProgress>(NumberOfTasks, FTaskProgress::FProgressReporter::CreateLambda([this, ProcessHandle](double InProgress)
 		{
-			ExecuteProcessProgressReporter(TakeId, InProgress);
+			ExecuteProcessProgressReporter(ProcessHandle.Get(), InProgress);
 		}));
 
-		Super::IngestTake(TakeId, IngestOptions, MoveTemp(TaskProgress));
+		Super::IngestTake(ProcessHandle.Get(), IngestOptions.Get(), MoveTemp(TaskProgress));
 	});
 }
 
diff --git a/Engine/Plugins/VirtualProduction/CaptureManager/CaptureManagerDevices/Source/MonoVideoLiveLinkDevice/Private/MonoVideoLiveLinkDevice.h b/Engine/Plugins/VirtualProduction/CaptureManager/CaptureManagerDevices/Source/MonoVideoLiveLinkDevice/Private/MonoVideoLiveLinkDevice.h
--- a/Engine/Plugins/VirtualProduction/CaptureManager/CaptureManagerDevices/Source/MonoVideoLiveLinkDevice/Private/MonoVideoLiveLinkDevice.h
+++ b/Engine/Plugins/VirtualProduction/CaptureManager/CaptureManagerDevices/Source/MonoVideoLiveLinkDevice/Private/MonoVideoLiveLinkDevice.h
@@ -85,7 +85,7 @@ private:
 	virtual FString GetFullTakePath(UE::CaptureManager::FTakeId InTakeId) const override;
 
 	//~ Begin ULiveLinkDeviceCapability_Ingest interface
-	virtual void RunIngestTake(int32 InTakeId, const UIngestCapability_Options* InIngestOptions) override;
+	virtual void RunConvertAndUploadTake(const UIngestCapability_ProcessHandle* InProcessHandle, const UIngestCapability_Options* InIngestOptions) override;
 	virtual void UpdateTakeList_Implementation(UIngestCapability_UpdateTakeListCallback* InCallback) override;
 	//~ End ULiveLinkDeviceCapability_Ingest interface
 
diff --git a/Engine/Plugins/VirtualProduction/CaptureManager/CaptureManagerDevices/Source/StereoVideoLiveLinkDevice/Private/StereoVideoLiveLinkDevice.cpp b/Engine/Plugins/VirtualProduction/CaptureManager/CaptureManagerDevices/Source/StereoVideoLiveLinkDevice/Private/StereoVideoLiveLinkDevice.cpp
--- a/Engine/Plugins/VirtualProduction/CaptureManager/CaptureManagerDevices/Source/StereoVideoLiveLinkDevice/Private/StereoVideoLiveLinkDevice.cpp
+++ b/Engine/Plugins/VirtualProduction/CaptureManager/CaptureManagerDevices/Source/StereoVideoLiveLinkDevice/Private/StereoVideoLiveLinkDevice.cpp
@@ -110,22 +110,22 @@ FString UStereoVideoLiveLinkDevice::GetFullTakePath(UE::CaptureManager::FTakeId
 	return FString();
 }
 
-void UStereoVideoLiveLinkDevice::RunIngestTake(int32 InTakeId, const UIngestCapability_Options* InIngestOptions)
+void UStereoVideoLiveLinkDevice::RunConvertAndUploadTake(const UIngestCapability_ProcessHandle* InProcessHandle, const UIngestCapability_Options* InIngestOptions)
 {
 	AsyncTask(ENamedThreads::AnyThread, [this,
-		TakeId = InTakeId,
-		IngestOptions = InIngestOptions]()
+		ProcessHandle = TStrongObjectPtr<const UIngestCapability_ProcessHandle>(InProcessHandle),
+		IngestOptions = TStrongObjectPtr<const UIngestCapability_Options>(InIngestOptions)]()
 		{
 			static constexpr uint32 NumberOfTasks = 2;
 
 			using namespace UE::CaptureManager;
 			TSharedPtr<FTaskProgress> TaskProgress
-				= MakeShared<FTaskProgress>(NumberOfTasks, FTaskProgress::FProgressReporter::CreateLambda([this, TakeId](double InProgress)
+				= MakeShared<FTaskProgress>(NumberOfTasks, FTaskProgress::FProgressReporter::CreateLambda([this, ProcessHandle](double InProgress)
 					{
-						ExecuteProcessProgressReporter(TakeId, InProgress);
+						ExecuteProcessProgressReporter(ProcessHandle.Get(), InProgress);
 					}));
 
-			Super::IngestTake(TakeId, IngestOptions, MoveTemp(TaskProgress));
+			Super::IngestTake(ProcessHandle.Get(), IngestOptions.Get(), MoveTemp(TaskProgress));
 		});
 }
 
diff --git a/Engine/Plugins/VirtualProduction/CaptureManager/CaptureManagerDevices/Source/StereoVideoLiveLinkDevice/Private/StereoVideoLiveLinkDevice.h b/Engine/Plugins/VirtualProduction/CaptureManager/CaptureManagerDevices/Source/StereoVideoLiveLinkDevice/Private/StereoVideoLiveLinkDevice.h
--- a/Engine/Plugins/VirtualProduction/CaptureManager/CaptureManagerDevices/Source/StereoVideoLiveLinkDevice/Private/StereoVideoLiveLinkDevice.h
+++ b/Engine/Plugins/VirtualProduction/CaptureManager/CaptureManagerDevices/Source/StereoVideoLiveLinkDevice/Private/StereoVideoLiveLinkDevice.h
@@ -141,7 +141,8 @@ private:
 	virtual FString GetFullTakePath(UE::CaptureManager::FTakeId InTakeId) const override;
 
 	//~ Begin ULiveLinkDeviceCapability_Ingest interface
-	virtual void RunIngestTake(int32 InTakeId, const UIngestCapability_Options* InIngestOptions) override;
+	virtual void RunConvertAndUploadTake(const UIngestCapability_ProcessHandle* InProcessHandle,
+										 const UIngestCapability_Options* InIngestOptions) override;
 	virtual void UpdateTakeList_Implementation(UIngestCapability_UpdateTakeListCallback* InCallback) override;
 	//~ End ULiveLinkDeviceCapability_Ingest interface
 
