Extract timecode from video and audio files if timecode is not defined in cptake or cptake does not exist

diff --git a/Engine/Plugins/VirtualProduction/CaptureManager/CaptureManagerApp/CaptureManagerApp.uplugin b/Engine/Plugins/VirtualProduction/CaptureManager/CaptureManagerApp/CaptureManagerApp.uplugin
--- a/Engine/Plugins/VirtualProduction/CaptureManager/CaptureManagerApp/CaptureManagerApp.uplugin
+++ b/Engine/Plugins/VirtualProduction/CaptureManager/CaptureManagerApp/CaptureManagerApp.uplugin
@@ -93,6 +93,10 @@
 		{
 			"Name": "LiveLinkDevice",
 			"Enabled": true
+		},
+		{
+			"Name": "ElectraPlayer",
+			"Enabled": true
 		}
 	]
 }
\ No newline at end of file
diff --git a/Engine/Plugins/VirtualProduction/CaptureManager/CaptureManagerApp/Source/IngestLiveLinkDevice/IngestLiveLinkDevice.Build.cs b/Engine/Plugins/VirtualProduction/CaptureManager/CaptureManagerApp/Source/IngestLiveLinkDevice/IngestLiveLinkDevice.Build.cs
--- a/Engine/Plugins/VirtualProduction/CaptureManager/CaptureManagerApp/Source/IngestLiveLinkDevice/IngestLiveLinkDevice.Build.cs
+++ b/Engine/Plugins/VirtualProduction/CaptureManager/CaptureManagerApp/Source/IngestLiveLinkDevice/IngestLiveLinkDevice.Build.cs
@@ -13,7 +13,8 @@ public class IngestLiveLinkDevice : ModuleRules
 		PublicDependencyModuleNames.AddRange(new string[]
 		{
 			"CaptureUtils",
-			"Core"
+			"Core",
+			"MediaAssets"
 		});
 
 		PrivateDependencyModuleNames.AddRange(new string[]
diff --git a/Engine/Plugins/VirtualProduction/CaptureManager/CaptureManagerApp/Source/IngestLiveLinkDevice/Private/Utils/CaptureExtractTimecode.cpp b/Engine/Plugins/VirtualProduction/CaptureManager/CaptureManagerApp/Source/IngestLiveLinkDevice/Private/Utils/CaptureExtractTimecode.cpp
new file mode 100644
--- /dev/null
+++ b/Engine/Plugins/VirtualProduction/CaptureManager/CaptureManagerApp/Source/IngestLiveLinkDevice/Private/Utils/CaptureExtractTimecode.cpp
@@ -0,0 +1,289 @@
+// Copyright Epic Games, Inc. All Rights Reserved.
+
+#include "Utils/CaptureExtractTimecode.h"
+
+#include "Audio.h"
+#include "Async/Async.h"
+#include "Misc/FileHelper.h"
+#include "MediaPlaylist.h"
+#include "Sound/SoundWaveTimecodeInfo.h"
+#include "UObject/Package.h"
+
+FCaptureExtractTimecode::FCaptureExtractTimecode(const FString& InFilePath)
+	: FilePath(InFilePath)
+{}
+
+FCaptureExtractTimecode::FTimecodeAndRateResult FCaptureExtractTimecode::Extract()
+{	
+	check(!FilePath.IsEmpty())
+
+	FTimecodeAndRateResult Result = MakeError(ECaptureExtractTimecodeError::UnhandledMedia);
+
+	FString FileExtension = FPaths::GetExtension(FilePath);
+	if (FileExtension == "wav")
+	{
+		Result = ExtractTimecodeFromWavBFWSection();
+	}
+	else
+	{
+		Result = ExtractTimecodeUsingElectraPlayer();
+	}
+
+	return Result;
+}
+
+FCaptureExtractTimecode::FTimecodeAndRateResult FCaptureExtractTimecode::ExtractTimecodeFromWavBFWSection()
+{
+	TArray<uint8> WAVData;
+	if (FFileHelper::LoadFileToArray(WAVData, *FilePath))
+	{
+		FWaveModInfo WAVInfo;
+		if (WAVInfo.ReadWaveInfo(WAVData.GetData(), WAVData.Num()))
+		{
+			FSoundWaveTimecodeInfo TimecodeInfo = *WAVInfo.TimecodeInfo.Get();
+			const double NumSecondsSinceMidnight = TimecodeInfo.GetNumSecondsSinceMidnight();
+			FTimecode AudioTimecode = FTimecode(NumSecondsSinceMidnight, TimecodeInfo.TimecodeRate, TimecodeInfo.bTimecodeIsDropFrame, /* InbRollover = */ true);
+			FTimecodeAndRate TimecodeAndRate{ AudioTimecode, TimecodeInfo.TimecodeRate };
+			return MakeValue(TimecodeAndRate);
+		}
+	}
+
+	return MakeError(ECaptureExtractTimecodeError::TimecodeNotFound);
+}
+
+FCaptureExtractTimecode::FTimecodeAndRateResult FCaptureExtractTimecode::ExtractTimecodeUsingElectraPlayer()
+{
+	//Some media players require waiting a frame before metadata can be read
+	bGetTimecodeDelay = true;
+
+	FTimecodeAndRatePromise Promise;
+	TFuture<FTimecodeAndRateResult> Future = Promise.GetFuture();
+
+	AsyncTask(ENamedThreads::GameThread, [this, &Promise]
+		{
+
+			UDesiredPlayerMediaSource* MediaSource = NewObject<UDesiredPlayerMediaSource>();
+			MediaSource->SetFilePath(FilePath);
+
+			// Create media player.
+			TStrongObjectPtr<UMediaPlayer> MediaPlayer = TStrongObjectPtr<UMediaPlayer>(
+				NewObject<UMediaPlayer>(GetTransientPackage(), MakeUniqueObjectName(GetTransientPackage(),
+					UMediaPlayer::StaticClass())));
+
+			check(MediaPlayer);
+
+			// Open the media.
+			MediaPlayer->PlayOnOpen = false;
+			FMediaPlayerOptions PlayerOptions;
+			PlayerOptions.SetAllAsOptional();
+			PlayerOptions.InternalCustomOptions.Emplace(MediaPlayerOptionValues::ParseTimecodeInfo(), FVariant());
+
+			if (!MediaPlayer->OpenSourceWithOptions(MediaSource, PlayerOptions))
+			{
+				Promise.SetValue(MakeError(ECaptureExtractTimecodeError::InternalError));
+				return;
+			}
+
+			Context = { MediaPlayer, Promise };
+
+			InitializeClockSink();
+		});
+
+	return Future.Get();
+}
+
+void FCaptureExtractTimecode::TickOutput(FTimespan DeltaTime, FTimespan Timecode)
+{
+#if WITH_EDITOR
+	if (Context.IsSet())
+	{
+		TStrongObjectPtr<UMediaPlayer>& MediaPlayer = Context.GetValue().MediaPlayer;
+
+		check(MediaPlayer);
+
+		if (MediaPlayer->HasError())
+		{
+			FTimecodeAndRatePromise& Promise = Context.GetValue().ExtractedTimecodeAndRatePromise;
+
+			Context.Reset();
+			DeinitializeClockSink();
+
+			Promise.SetValue(MakeError(ECaptureExtractTimecodeError::InternalError));
+			return;
+		}
+
+		const int32 TimeoutPeriod = 3;
+		const FDateTime WaitStart = FDateTime::Now();
+		// Player is not ready if we can't get duration. Delay until duration is not zero.
+		if (bGetTimecodeDelay)
+		{
+			const FTimespan Duration = MediaPlayer->GetDuration();
+			if (Duration == 0)
+			{
+				if ((FDateTime::Now() - WaitStart).GetSeconds() > TimeoutPeriod)
+				{
+					FTimecodeAndRatePromise& Promise = Context.GetValue().ExtractedTimecodeAndRatePromise;
+					MediaPlayer->Close();
+					Context.Reset();
+					DeinitializeClockSink();
+					Promise.SetValue(MakeValue(FTimecodeAndRate()));
+					return;
+				}
+
+				// Keep delaying
+				return;
+			}
+
+			bGetTimecodeDelay = false;
+		}
+
+		TValueOrError<FTimecode, ECaptureExtractTimecodeError> TimecodeResult = GetTimecode(MediaPlayer);
+		TValueOrError<FFrameRate, ECaptureExtractTimecodeError> TimecodeRateResult = GetTimecodeRate(MediaPlayer);
+		MediaPlayer->Close();
+
+		FTimecodeAndRatePromise& Promise = Context.GetValue().ExtractedTimecodeAndRatePromise;
+
+		Context.Reset();
+		DeinitializeClockSink();
+
+		FTimecodeAndRate TimecodeAndRate;
+		if (TimecodeResult.IsValid())
+		{
+			TimecodeAndRate.MaybeTimecode = TimecodeResult.StealValue();
+		}
+		if (TimecodeRateResult.IsValid())
+		{
+			TimecodeAndRate.MaybeTimecodeRate = TimecodeRateResult.StealValue();
+		}
+		Promise.SetValue(MakeValue(MoveTemp(TimecodeAndRate)));
+	}
+#endif // WITH_EDITOR
+}
+
+void FCaptureExtractTimecode::InitializeClockSink()
+{
+	IMediaModule* MediaModule = FModuleManager::LoadModulePtr<IMediaModule>("Media");
+	if (MediaModule != nullptr)
+	{
+		MediaModule->GetClock().AddSink(this->AsShared());
+	}
+}
+
+void FCaptureExtractTimecode::DeinitializeClockSink()
+{
+	IMediaModule* MediaModule = FModuleManager::GetModulePtr<IMediaModule>("Media");
+	if (MediaModule != nullptr)
+	{
+		MediaModule->GetClock().RemoveSink(this->AsShared());
+	}
+}
+
+TValueOrError<FTimecode, FCaptureExtractTimecode::ECaptureExtractTimecodeError> FCaptureExtractTimecode::GetTimecode(TStrongObjectPtr<UMediaPlayer>& MediaPlayer)
+{
+	// Timecode from the source. Either player's meta data or media source.
+	const FVariant Value = MediaPlayer->GetMediaInfo(UMediaPlayer::MediaInfoNameStartTimecodeValue.Resolve());
+	FString TimecodeString;
+
+	if (!Value.IsEmpty())
+	{
+		TimecodeString = Value.GetValue<FString>();
+	}
+	else
+	{
+
+		if (UMediaPlaylist* Playlist = MediaPlayer->GetPlaylist())
+		{
+			if (const UMediaSource* MediaSource = Playlist->Get(MediaPlayer->GetPlaylistIndex()))
+			{
+				TimecodeString = MediaSource->GetMediaOption(UMediaPlayer::MediaInfoNameStartTimecodeValue.Resolve(), FString());
+			}
+		}
+	}
+
+	if (TimecodeString.IsEmpty())
+	{
+		return MakeError(ECaptureExtractTimecodeError::TimecodeNotFound);
+	}
+	else
+	{
+		const TOptional<FTimecode> Timecode = FTimecode::ParseTimecode(*TimecodeString);
+		if (!Timecode.IsSet())
+		{
+			return MakeError(ECaptureExtractTimecodeError::UnableToParseTimecode);
+		}
+
+		return MakeValue(Timecode.GetValue());
+	}
+}
+
+TValueOrError<FFrameRate, FCaptureExtractTimecode::ECaptureExtractTimecodeError> FCaptureExtractTimecode::GetTimecodeRate(TStrongObjectPtr<UMediaPlayer>& MediaPlayer)
+{
+	// Timecode from the source. Either player's meta data or media source.
+	const FVariant Value = MediaPlayer->GetMediaInfo(UMediaPlayer::MediaInfoNameStartTimecodeFrameRate.Resolve());
+	FString TimecodeRateString;
+
+	if (!Value.IsEmpty())
+	{
+		TimecodeRateString = Value.GetValue<FString>();
+	}
+	else
+	{
+
+		if (UMediaPlaylist* Playlist = MediaPlayer->GetPlaylist())
+		{
+			if (const UMediaSource* MediaSource = Playlist->Get(MediaPlayer->GetPlaylistIndex()))
+			{
+				TimecodeRateString = MediaSource->GetMediaOption(UMediaPlayer::MediaInfoNameStartTimecodeFrameRate.Resolve(), FString());
+			}
+		}
+	}
+
+	if (TimecodeRateString.IsEmpty())
+	{
+		return MakeError(ECaptureExtractTimecodeError::TimecodeRateNotFound);
+	}
+	else
+	{
+		// TimecodeRateString is made out of FFrameRate::ToPrettyText()
+		// It could be "{0} fps" or "{0} s"
+
+		FFrameRate TimecodeRate;
+
+		FString Left;
+		FString Right;
+		TimecodeRateString.Split(" ", &Left, &Right);
+
+		double Number = FCString::Atoi(*Left);
+		uint32 IntNumber;
+		uint32 Multiplier = 1;
+
+
+		for (; Multiplier <= 10000; Multiplier *= 10)
+		{
+			double TmpNumber = Number * Multiplier;
+			IntNumber = FMath::RoundToInt32(TmpNumber);
+			if (FMath::Abs(TmpNumber - IntNumber) < 0.01)
+			{
+				break;
+			}
+		}
+
+		uint32 Nominator = 0;
+		uint32 Denominator = 0;
+
+		if (Right == "fps")
+		{
+			Nominator = IntNumber;
+			Denominator = Multiplier;
+		}
+		else if (Right == "s")
+		{
+			Nominator = Multiplier;
+			Denominator = IntNumber;
+		}
+
+		TimecodeRate = FFrameRate(Nominator, Denominator);
+
+		return MakeValue(TimecodeRate);
+	}
+}
diff --git a/Engine/Plugins/VirtualProduction/CaptureManager/CaptureManagerApp/Source/IngestLiveLinkDevice/Public/Utils/CaptureExtractTimecode.h b/Engine/Plugins/VirtualProduction/CaptureManager/CaptureManagerApp/Source/IngestLiveLinkDevice/Public/Utils/CaptureExtractTimecode.h
new file mode 100644
--- /dev/null
+++ b/Engine/Plugins/VirtualProduction/CaptureManager/CaptureManagerApp/Source/IngestLiveLinkDevice/Public/Utils/CaptureExtractTimecode.h
@@ -0,0 +1,95 @@
+// Copyright Epic Games, Inc. All Rights Reserved.
+
+#pragma once
+
+#include "FileMediaSource.h"
+
+#include "MediaSample.h"
+#include "MediaPlayer.h"
+
+#include "Templates/ValueOrError.h"
+
+#include "IMediaClock.h"
+#include "IMediaClockSink.h"
+#include "IMediaModule.h"
+
+#include "CaptureExtractTimecode.generated.h"
+
+class UMediaPlayer;
+
+#if WITH_EDITOR
+
+UCLASS(BlueprintType, MinimalAPI)
+class UDesiredPlayerMediaSource : public UFileMediaSource
+{
+	GENERATED_BODY()
+
+public:
+
+	//~ IMediaOptions interface
+
+	virtual FName GetDesiredPlayerName() const override
+	{
+		return TEXT("ElectraPlayer");
+	}
+
+};
+#endif // WITH_EDITOR
+
+class INGESTLIVELINKDEVICE_API FCaptureExtractTimecode final : public TSharedFromThis<FCaptureExtractTimecode>, public IMediaClockSink
+{
+public:
+
+	struct FTimecodeAndRate
+	{
+		TOptional<FTimecode> MaybeTimecode;
+		TOptional<FFrameRate> MaybeTimecodeRate;
+	};
+
+	enum ECaptureExtractTimecodeError : int32
+	{
+		InternalError = -1,
+		TimecodeNotFound = 1,
+		UnableToParseTimecode = 2,
+		UnableToParseTimecodeRate = 3,
+		TimecodeRateNotFound = 4,
+		UnhandledMedia = 5,
+	};
+
+	using FTimecodeAndRateResult = TValueOrError<FTimecodeAndRate, ECaptureExtractTimecodeError>;
+
+	FCaptureExtractTimecode(const FString& InFilePath);
+
+	FTimecodeAndRateResult Extract();
+
+private:
+
+	FString FilePath;
+
+	using FTimecodeAndRatePromise = TPromise<FTimecodeAndRateResult>;
+
+	struct FTimecodeExtractionContext
+	{
+		TStrongObjectPtr<UMediaPlayer> MediaPlayer;
+		FTimecodeAndRatePromise& ExtractedTimecodeAndRatePromise;
+	};
+
+	TOptional<FTimecodeExtractionContext> Context;
+
+#if WITH_EDITOR
+	TValueOrError<FTimecode, ECaptureExtractTimecodeError> GetTimecode(TStrongObjectPtr<UMediaPlayer>& MediaPlayer);
+	TValueOrError<FFrameRate, ECaptureExtractTimecodeError> GetTimecodeRate(TStrongObjectPtr<UMediaPlayer>& MediaPlayer);
+#endif
+	void InitializeClockSink();
+	void DeinitializeClockSink();
+
+	FTimecodeAndRateResult ExtractTimecodeFromWavBFWSection();
+	FTimecodeAndRateResult ExtractTimecodeUsingElectraPlayer();
+	
+	bool bGetTimecodeDelay;
+	
+	/**
+	 * Called from the media clock.
+	 */
+	virtual void TickOutput(FTimespan DeltaTime, FTimespan Timecode) override;
+};
diff --git a/Engine/Plugins/VirtualProduction/CaptureManager/CaptureManagerDevices/Source/CPSLiveLinkDevice/Private/LiveLinkFaceDevice.cpp b/Engine/Plugins/VirtualProduction/CaptureManager/CaptureManagerDevices/Source/CPSLiveLinkDevice/Private/LiveLinkFaceDevice.cpp
--- a/Engine/Plugins/VirtualProduction/CaptureManager/CaptureManagerDevices/Source/CPSLiveLinkDevice/Private/LiveLinkFaceDevice.cpp
+++ b/Engine/Plugins/VirtualProduction/CaptureManager/CaptureManagerDevices/Source/CPSLiveLinkDevice/Private/LiveLinkFaceDevice.cpp
@@ -14,6 +14,7 @@
 
 #include "LiveLinkFaceMetadata.h"
 #include "StereoCameraMetadataParseUtils.h"
+#include "Utils/CaptureExtractTimecode.h"
 
 #include "HAL/FileManager.h"
 
@@ -406,8 +407,8 @@ void ULiveLinkFaceDevice::OnExportFinished(UE::CaptureManager::TProtocolResult<v
 	if (InResult.HasValue())
 	{
 		// Blocking function
-		TOptional<FTakeMetadata> TakeMetadata = ParseTake(DownloadedStorage, InTakeName);
-		if (TakeMetadata.IsSet())
+		TOptional<FTakeMetadata> MaybeTakeMetadata = ParseTake(DownloadedStorage, InTakeName);
+		if (MaybeTakeMetadata.IsSet())
 		{
 			FTakeId TakeId = InProcessHandle->GetTakeId();
 
@@ -415,7 +416,11 @@ void ULiveLinkFaceDevice::OnExportFinished(UE::CaptureManager::TProtocolResult<v
 			DownloadedTakes.Add(TakeId, DownloadedTake);
 			Lock.Unlock();
 
-			UpdateTake(TakeId, MoveTemp(TakeMetadata.GetValue()));
+			FTakeMetadata TakeMetadata = MaybeTakeMetadata.GetValue();
+			
+			ExtractTimecodeIfNotSet(TakeMetadata);
+			
+			UpdateTake(TakeId, MoveTemp(TakeMetadata));
 
 			ExecuteProcessFinishedReporter(InProcessHandle.Get(), MakeValue());
 		}
@@ -544,4 +549,62 @@ void ULiveLinkFaceDevice::RemoveDownloadedTakeData(const UE::CaptureManager::FTa
 	}
 
 	IFileManager::Get().DeleteDirectory(**DownloadedTake, false, true);
-}
\ No newline at end of file
+}
+
+void ULiveLinkFaceDevice::ExtractTimecodeIfNotSet(FTakeMetadata& InOutTakeMetadata)
+{
+	for (auto& Video : InOutTakeMetadata.Video)
+	{
+		if (!Video.TimecodeStart.IsSet())
+		{
+			TSharedPtr<FCaptureExtractTimecode> Extractor = MakeShareable(new FCaptureExtractTimecode(Video.Path));
+			FCaptureExtractTimecode::FTimecodeAndRateResult TimecodeAndRateResult = Extractor->Extract();
+			if (TimecodeAndRateResult.IsValid())
+			{
+				FCaptureExtractTimecode::FTimecodeAndRate TimecodeAndRate = TimecodeAndRateResult.GetValue();
+				if (TimecodeAndRate.MaybeTimecode.IsSet())
+				{
+					Video.TimecodeStart = TimecodeAndRate.MaybeTimecode.GetValue().ToString();
+				}
+			}
+		}
+	}
+
+	for (auto& Depth : InOutTakeMetadata.Depth)
+	{
+		if (!Depth.TimecodeStart.IsSet())
+		{
+			TSharedPtr<FCaptureExtractTimecode> Extractor = MakeShareable(new FCaptureExtractTimecode(Depth.Path));
+			FCaptureExtractTimecode::FTimecodeAndRateResult TimecodeAndRateResult = Extractor->Extract();
+			if (TimecodeAndRateResult.IsValid())
+			{
+				FCaptureExtractTimecode::FTimecodeAndRate TimecodeAndRate = TimecodeAndRateResult.GetValue();
+				if (TimecodeAndRate.MaybeTimecode.IsSet())
+				{
+					Depth.TimecodeStart = TimecodeAndRate.MaybeTimecode.GetValue().ToString();
+				}
+			}
+		}
+	}
+
+	for (auto& Audio : InOutTakeMetadata.Audio)
+	{
+		if (!Audio.TimecodeStart.IsSet() && !Audio.TimecodeRate.IsSet())
+		{
+			TSharedPtr<FCaptureExtractTimecode> Extractor = MakeShareable(new FCaptureExtractTimecode(Audio.Path));
+			FCaptureExtractTimecode::FTimecodeAndRateResult TimecodeAndRateResult = Extractor->Extract();
+			if (TimecodeAndRateResult.IsValid())
+			{
+				FCaptureExtractTimecode::FTimecodeAndRate TimecodeAndRate = TimecodeAndRateResult.GetValue();
+				if (TimecodeAndRate.MaybeTimecode.IsSet())
+				{
+					Audio.TimecodeStart = TimecodeAndRate.MaybeTimecode.GetValue().ToString();
+				}
+				if (TimecodeAndRate.MaybeTimecodeRate.IsSet())
+				{
+					Audio.TimecodeRate = static_cast<float>(TimecodeAndRate.MaybeTimecodeRate.GetValue().AsDecimal());
+				}
+			}
+		}
+	}
+}
diff --git a/Engine/Plugins/VirtualProduction/CaptureManager/CaptureManagerDevices/Source/CPSLiveLinkDevice/Private/LiveLinkFaceDevice.h b/Engine/Plugins/VirtualProduction/CaptureManager/CaptureManagerDevices/Source/CPSLiveLinkDevice/Private/LiveLinkFaceDevice.h
--- a/Engine/Plugins/VirtualProduction/CaptureManager/CaptureManagerDevices/Source/CPSLiveLinkDevice/Private/LiveLinkFaceDevice.h
+++ b/Engine/Plugins/VirtualProduction/CaptureManager/CaptureManagerDevices/Source/CPSLiveLinkDevice/Private/LiveLinkFaceDevice.h
@@ -102,6 +102,8 @@ private:
 
 	void RemoveDownloadedTakeData(const UE::CaptureManager::FTakeId InTakeId);
 
+	void ExtractTimecodeIfNotSet(FTakeMetadata& InOutTakeMetadata);
+
 	TSharedPtr<UE::CaptureManager::FCPSDevice> Device;
 
 	mutable FCriticalSection DownloadedTakesMutex;
diff --git a/Engine/Plugins/VirtualProduction/CaptureManager/CaptureManagerDevices/Source/FilesystemLiveLinkDevice/Private/FilesystemLiveLinkDevice.cpp b/Engine/Plugins/VirtualProduction/CaptureManager/CaptureManagerDevices/Source/FilesystemLiveLinkDevice/Private/FilesystemLiveLinkDevice.cpp
--- a/Engine/Plugins/VirtualProduction/CaptureManager/CaptureManagerDevices/Source/FilesystemLiveLinkDevice/Private/FilesystemLiveLinkDevice.cpp
+++ b/Engine/Plugins/VirtualProduction/CaptureManager/CaptureManagerDevices/Source/FilesystemLiveLinkDevice/Private/FilesystemLiveLinkDevice.cpp
@@ -11,6 +11,7 @@
 #include "StereoCameraMetadataParseUtils.h"
 
 #include "Utils/IngestLiveLinkDeviceUtils.h"
+#include "Utils/CaptureExtractTimecode.h"
 
 DEFINE_LOG_CATEGORY(LogGenericLiveLinkArchiveDevice);
 
@@ -237,48 +238,72 @@ TOptional<FTakeMetadata> UFilesystemLiveLinkDevice::ReadTake(const FString& InCu
 			TakeMetadata.Thumbnail = FTakeThumbnailData(FPaths::ConvertRelativePathToFull(CurrentDirectory, *ThumbnailPath));
 		}
 	}
-
-	if (!TakeMetadata.Video.IsEmpty())
+	
+	for (FTakeMetadata::FVideo& Video : TakeMetadata.Video)
 	{
-		for (FTakeMetadata::FVideo& Video : TakeMetadata.Video)
+		if (FPaths::IsRelative(*Video.Path))
+		{
+			Video.Path = FPaths::ConvertRelativePathToFull(CurrentDirectory, *Video.Path);
+		}
+
+		TSharedPtr<FCaptureExtractTimecode> Extractor = MakeShareable(new FCaptureExtractTimecode(Video.Path));
+		FCaptureExtractTimecode::FTimecodeAndRateResult TimecodeAndRateResult = Extractor->Extract();
+		if (TimecodeAndRateResult.IsValid())
 		{
-			if (FPaths::IsRelative(*Video.Path))
+			FCaptureExtractTimecode::FTimecodeAndRate TimecodeAndRate = TimecodeAndRateResult.GetValue();
+			if (!Video.TimecodeStart.IsSet() && TimecodeAndRate.MaybeTimecode.IsSet())
 			{
-				Video.Path = FPaths::ConvertRelativePathToFull(CurrentDirectory, *Video.Path);
+				Video.TimecodeStart = TimecodeAndRate.MaybeTimecode.GetValue().ToString();
 			}
 		}
 	}
-
-	if (!TakeMetadata.Depth.IsEmpty())
+	for (FTakeMetadata::FVideo& Depth : TakeMetadata.Depth)
 	{
-		for (FTakeMetadata::FVideo& Depth : TakeMetadata.Depth)
+		if (FPaths::IsRelative(*Depth.Path))
+		{
+			Depth.Path = FPaths::ConvertRelativePathToFull(CurrentDirectory, *Depth.Path);
+		}
+
+		TSharedPtr<FCaptureExtractTimecode> Extractor = MakeShareable(new FCaptureExtractTimecode(Depth.Path));
+		FCaptureExtractTimecode::FTimecodeAndRateResult TimecodeAndRateResult = Extractor->Extract();
+		if (TimecodeAndRateResult.IsValid())
 		{
-			if (FPaths::IsRelative(*Depth.Path))
+			FCaptureExtractTimecode::FTimecodeAndRate TimecodeAndRate = TimecodeAndRateResult.GetValue();
+			if (!Depth.TimecodeStart.IsSet() && TimecodeAndRate.MaybeTimecode.IsSet())
 			{
-				Depth.Path = FPaths::ConvertRelativePathToFull(CurrentDirectory, *Depth.Path);
+				Depth.TimecodeStart = TimecodeAndRate.MaybeTimecode.GetValue().ToString();
 			}
 		}
 	}
 
-	if (!TakeMetadata.Audio.IsEmpty())
+	for (FTakeMetadata::FAudio& Audio : TakeMetadata.Audio)
 	{
-		for (FTakeMetadata::FAudio& Audio : TakeMetadata.Audio)
+		if (FPaths::IsRelative(*Audio.Path))
 		{
-			if (FPaths::IsRelative(*Audio.Path))
+			Audio.Path = FPaths::ConvertRelativePathToFull(CurrentDirectory, *Audio.Path);
+		}
+
+		TSharedPtr<FCaptureExtractTimecode> Extractor = MakeShareable(new FCaptureExtractTimecode(Audio.Path));
+		FCaptureExtractTimecode::FTimecodeAndRateResult TimecodeAndRateResult = Extractor->Extract();
+		if (TimecodeAndRateResult.IsValid())
+		{
+			FCaptureExtractTimecode::FTimecodeAndRate TimecodeAndRate = TimecodeAndRateResult.GetValue();
+			if (TimecodeAndRate.MaybeTimecode.IsSet())
+			{
+				Audio.TimecodeStart = TimecodeAndRate.MaybeTimecode.GetValue().ToString();
+			}
+			if (TimecodeAndRate.MaybeTimecodeRate.IsSet())
 			{
-				Audio.Path = FPaths::ConvertRelativePathToFull(CurrentDirectory, *Audio.Path);
+				Audio.TimecodeRate = static_cast<float>(TimecodeAndRate.MaybeTimecodeRate.GetValue().AsDecimal());
 			}
 		}
 	}
-
-	if (!TakeMetadata.Calibration.IsEmpty())
+	
+	for (FTakeMetadata::FCalibration& Calib : TakeMetadata.Calibration)
 	{
-		for (FTakeMetadata::FCalibration& Calib : TakeMetadata.Calibration)
+		if (FPaths::IsRelative(*Calib.Path))
 		{
-			if (FPaths::IsRelative(*Calib.Path))
-			{
-				Calib.Path = FPaths::ConvertRelativePathToFull(CurrentDirectory, *Calib.Path);
-			}
+			Calib.Path = FPaths::ConvertRelativePathToFull(CurrentDirectory, *Calib.Path);
 		}
 	}
 
diff --git a/Engine/Plugins/VirtualProduction/CaptureManager/CaptureManagerDevices/Source/MonoVideoLiveLinkDevice/Private/MonoVideoLiveLinkDevice.cpp b/Engine/Plugins/VirtualProduction/CaptureManager/CaptureManagerDevices/Source/MonoVideoLiveLinkDevice/Private/MonoVideoLiveLinkDevice.cpp
--- a/Engine/Plugins/VirtualProduction/CaptureManager/CaptureManagerDevices/Source/MonoVideoLiveLinkDevice/Private/MonoVideoLiveLinkDevice.cpp
+++ b/Engine/Plugins/VirtualProduction/CaptureManager/CaptureManagerDevices/Source/MonoVideoLiveLinkDevice/Private/MonoVideoLiveLinkDevice.cpp
@@ -11,6 +11,7 @@
 #include "Misc/MessageDialog.h"
 #include "Utils/TakeDiscoveryExpressionParser.h"
 #include "Utils/VideoDeviceThumbnailExtractor.h"
+#include "Utils/CaptureExtractTimecode.h"
 
 #include "VideoLiveLinkDeviceLog.h"
 
@@ -308,6 +309,21 @@ TOptional<FTakeMetadata> UMonoVideoLiveLinkDevice::ReadTake(const FString& InCur
 	Video.Path = InCurrentTakeFile;
 	Video.PathType = FTakeMetadata::FVideo::EPathType::File;
 
+	TSharedPtr<FCaptureExtractTimecode> Extractor = MakeShareable(new FCaptureExtractTimecode(Video.Path));
+	FCaptureExtractTimecode::FTimecodeAndRateResult TimecodeAndRateResult = Extractor->Extract();
+	if (TimecodeAndRateResult.IsValid())
+	{
+		FCaptureExtractTimecode::FTimecodeAndRate TimecodeAndRate = TimecodeAndRateResult.GetValue();
+		if (TimecodeAndRate.MaybeTimecode.IsSet())
+		{
+			Video.TimecodeStart = TimecodeAndRate.MaybeTimecode.GetValue().ToString();
+		}
+		if (TimecodeAndRate.MaybeTimecodeRate.IsSet())
+		{
+			Video.FrameRate = static_cast<float>(TimecodeAndRate.MaybeTimecodeRate.GetValue().AsDecimal());
+		}
+	}
+
 	TakeMetadata.Video.Add(MoveTemp(Video));
 
 	// Check if the video file contains audio stream by creating a audio reader from the file
diff --git a/Engine/Plugins/VirtualProduction/CaptureManager/CaptureManagerDevices/Source/StereoVideoLiveLinkDevice/Private/StereoVideoLiveLinkDevice.cpp b/Engine/Plugins/VirtualProduction/CaptureManager/CaptureManagerDevices/Source/StereoVideoLiveLinkDevice/Private/StereoVideoLiveLinkDevice.cpp
--- a/Engine/Plugins/VirtualProduction/CaptureManager/CaptureManagerDevices/Source/StereoVideoLiveLinkDevice/Private/StereoVideoLiveLinkDevice.cpp
+++ b/Engine/Plugins/VirtualProduction/CaptureManager/CaptureManagerDevices/Source/StereoVideoLiveLinkDevice/Private/StereoVideoLiveLinkDevice.cpp
@@ -13,6 +13,7 @@
 #include "Misc/MessageDialog.h"
 #include "Utils/TakeDiscoveryExpressionParser.h"
 #include "Utils/VideoDeviceThumbnailExtractor.h"
+#include "Utils/CaptureExtractTimecode.h"
 
 #include "VideoLiveLinkDeviceLog.h"
 
@@ -520,7 +521,18 @@ TOptional<FTakeMetadata> UStereoVideoLiveLinkDevice::CreateTakeMetadata(const FT
 				}
 
 				Video.Path = TakeComponent.Path;
-
+				
+				FString MediaFilePath = FPaths::ConvertRelativePathToFull(InTakeWithComponents.TakePath, Video.Path);
+				TSharedPtr<FCaptureExtractTimecode> Extractor = MakeShareable(new FCaptureExtractTimecode(MediaFilePath));
+				FCaptureExtractTimecode::FTimecodeAndRateResult TimecodeAndRateResult = Extractor->Extract();
+				if (TimecodeAndRateResult.IsValid())
+				{
+					FCaptureExtractTimecode::FTimecodeAndRate TimecodeAndRate = TimecodeAndRateResult.GetValue();
+					if (TimecodeAndRate.MaybeTimecode.IsSet())
+					{
+						Video.TimecodeStart = TimecodeAndRate.MaybeTimecode.GetValue().ToString();
+					}
+				}
 			}
 			else
 			{
@@ -582,6 +594,22 @@ TOptional<FTakeMetadata> UStereoVideoLiveLinkDevice::CreateTakeMetadata(const FT
 			Audio.Path = TakeComponent.Path;
 			Audio.Duration = 0;
 
+			FString MediaFilePath = FPaths::ConvertRelativePathToFull(InTakeWithComponents.TakePath, Audio.Path);
+			TSharedPtr<FCaptureExtractTimecode> Extractor = MakeShareable(new FCaptureExtractTimecode(MediaFilePath));
+			FCaptureExtractTimecode::FTimecodeAndRateResult TimecodeAndRateResult = Extractor->Extract();
+			if (TimecodeAndRateResult.IsValid())
+			{
+				FCaptureExtractTimecode::FTimecodeAndRate TimecodeAndRate = TimecodeAndRateResult.GetValue();
+				if (TimecodeAndRate.MaybeTimecode.IsSet())
+				{
+					Audio.TimecodeStart = TimecodeAndRate.MaybeTimecode.GetValue().ToString();
+				}
+				if (TimecodeAndRate.MaybeTimecodeRate.IsSet())
+				{
+					Audio.TimecodeRate = static_cast<float>(TimecodeAndRate.MaybeTimecodeRate.GetValue().AsDecimal());
+				}
+			}
+
 			TakeMetadata.Audio.Add(MoveTemp(Audio));
 		}
 	}
