Fix crash on wav file timecode extraction
Also fix wait bug and do not set video frame rate in mono device

diff --git a/Engine/Plugins/VirtualProduction/CaptureManager/CaptureManagerApp/Source/IngestLiveLinkDevice/Private/Utils/CaptureExtractTimecode.cpp b/Engine/Plugins/VirtualProduction/CaptureManager/CaptureManagerApp/Source/IngestLiveLinkDevice/Private/Utils/CaptureExtractTimecode.cpp
--- a/Engine/Plugins/VirtualProduction/CaptureManager/CaptureManagerApp/Source/IngestLiveLinkDevice/Private/Utils/CaptureExtractTimecode.cpp
+++ b/Engine/Plugins/VirtualProduction/CaptureManager/CaptureManagerApp/Source/IngestLiveLinkDevice/Private/Utils/CaptureExtractTimecode.cpp
@@ -40,11 +40,14 @@ FCaptureExtractTimecode::FTimecodeAndRateResult FCaptureExtractTimecode::Extract
 		FWaveModInfo WAVInfo;
 		if (WAVInfo.ReadWaveInfo(WAVData.GetData(), WAVData.Num()))
 		{
-			FSoundWaveTimecodeInfo TimecodeInfo = *WAVInfo.TimecodeInfo.Get();
-			const double NumSecondsSinceMidnight = TimecodeInfo.GetNumSecondsSinceMidnight();
-			FTimecode AudioTimecode = FTimecode(NumSecondsSinceMidnight, TimecodeInfo.TimecodeRate, TimecodeInfo.bTimecodeIsDropFrame, /* InbRollover = */ true);
-			FTimecodeAndRate TimecodeAndRate{ AudioTimecode, TimecodeInfo.TimecodeRate };
-			return MakeValue(TimecodeAndRate);
+			if (WAVInfo.TimecodeInfo.IsValid())
+			{
+				FSoundWaveTimecodeInfo TimecodeInfo = *WAVInfo.TimecodeInfo.Get();
+				const double NumSecondsSinceMidnight = TimecodeInfo.GetNumSecondsSinceMidnight();
+				FTimecode AudioTimecode = FTimecode(NumSecondsSinceMidnight, TimecodeInfo.TimecodeRate, TimecodeInfo.bTimecodeIsDropFrame, /* InbRollover = */ true);
+				FTimecodeAndRate TimecodeAndRate{ AudioTimecode, TimecodeInfo.TimecodeRate };
+				return MakeValue(TimecodeAndRate);
+			}
 		}
 	}
 
@@ -85,6 +88,10 @@ FCaptureExtractTimecode::FTimecodeAndRateResult FCaptureExtractTimecode::Extract
 			}
 
 			Context = { MediaPlayer, Promise };
+			
+			// Sometimes it is necessary to wait a couple of ticks for the MediaPlayer to be ready.
+			// Log start time so that we can timeout if initialization is taking too long.
+			MediaPlayerWaitStart = FDateTime::Now();
 
 			InitializeClockSink();
 		});
@@ -112,15 +119,13 @@ void FCaptureExtractTimecode::TickOutput(FTimespan DeltaTime, FTimespan Timecode
 			return;
 		}
 
-		const int32 TimeoutPeriod = 3;
-		const FDateTime WaitStart = FDateTime::Now();
 		// Player is not ready if we can't get duration. Delay until duration is not zero.
 		if (bGetTimecodeDelay)
 		{
 			const FTimespan Duration = MediaPlayer->GetDuration();
 			if (Duration == 0)
 			{
-				if ((FDateTime::Now() - WaitStart).GetSeconds() > TimeoutPeriod)
+				if ((FDateTime::Now() - MediaPlayerWaitStart).GetSeconds() > TimeoutPeriod)
 				{
 					FTimecodeAndRatePromise& Promise = Context.GetValue().ExtractedTimecodeAndRatePromise;
 					MediaPlayer->Close();
@@ -146,16 +151,21 @@ void FCaptureExtractTimecode::TickOutput(FTimespan DeltaTime, FTimespan Timecode
 		Context.Reset();
 		DeinitializeClockSink();
 
+		FTimecodeAndRateResult Result = MakeError(ECaptureExtractTimecodeError::TimecodeNotFound);
 		FTimecodeAndRate TimecodeAndRate;
 		if (TimecodeResult.IsValid())
 		{
 			TimecodeAndRate.MaybeTimecode = TimecodeResult.StealValue();
+		
+			if (TimecodeRateResult.IsValid())
+			{
+				TimecodeAndRate.MaybeTimecodeRate = TimecodeRateResult.StealValue();
+			}
+			
+			Result = MakeValue(MoveTemp(TimecodeAndRate));
 		}
-		if (TimecodeRateResult.IsValid())
-		{
-			TimecodeAndRate.MaybeTimecodeRate = TimecodeRateResult.StealValue();
-		}
-		Promise.SetValue(MakeValue(MoveTemp(TimecodeAndRate)));
+
+		Promise.SetValue(MoveTemp(Result));
 	}
 #endif // WITH_EDITOR
 }
diff --git a/Engine/Plugins/VirtualProduction/CaptureManager/CaptureManagerApp/Source/IngestLiveLinkDevice/Public/Utils/CaptureExtractTimecode.h b/Engine/Plugins/VirtualProduction/CaptureManager/CaptureManagerApp/Source/IngestLiveLinkDevice/Public/Utils/CaptureExtractTimecode.h
--- a/Engine/Plugins/VirtualProduction/CaptureManager/CaptureManagerApp/Source/IngestLiveLinkDevice/Public/Utils/CaptureExtractTimecode.h
+++ b/Engine/Plugins/VirtualProduction/CaptureManager/CaptureManagerApp/Source/IngestLiveLinkDevice/Public/Utils/CaptureExtractTimecode.h
@@ -87,7 +87,9 @@ private:
 	FTimecodeAndRateResult ExtractTimecodeUsingElectraPlayer();
 	
 	bool bGetTimecodeDelay;
-	
+	const int32 TimeoutPeriod = 3;
+	FDateTime MediaPlayerWaitStart;
+
 	/**
 	 * Called from the media clock.
 	 */
diff --git a/Engine/Plugins/VirtualProduction/CaptureManager/CaptureManagerDevices/Source/MonoVideoLiveLinkDevice/Private/MonoVideoLiveLinkDevice.cpp b/Engine/Plugins/VirtualProduction/CaptureManager/CaptureManagerDevices/Source/MonoVideoLiveLinkDevice/Private/MonoVideoLiveLinkDevice.cpp
--- a/Engine/Plugins/VirtualProduction/CaptureManager/CaptureManagerDevices/Source/MonoVideoLiveLinkDevice/Private/MonoVideoLiveLinkDevice.cpp
+++ b/Engine/Plugins/VirtualProduction/CaptureManager/CaptureManagerDevices/Source/MonoVideoLiveLinkDevice/Private/MonoVideoLiveLinkDevice.cpp
@@ -318,10 +318,6 @@ TOptional<FTakeMetadata> UMonoVideoLiveLinkDevice::ReadTake(const FString& InCur
 		{
 			Video.TimecodeStart = TimecodeAndRate.MaybeTimecode.GetValue().ToString();
 		}
-		if (TimecodeAndRate.MaybeTimecodeRate.IsSet())
-		{
-			Video.FrameRate = static_cast<float>(TimecodeAndRate.MaybeTimecodeRate.GetValue().AsDecimal());
-		}
 	}
 
 	TakeMetadata.Video.Add(MoveTemp(Video));
