[CaptureManager] Default ingest locations to UserDir

diff --git a/Engine/Plugins/VirtualProduction/CaptureManager/CaptureManagerApp/Source/CaptureManagerSettings/Private/Settings/CaptureManagerSettings.cpp b/Engine/Plugins/VirtualProduction/CaptureManager/CaptureManagerApp/Source/CaptureManagerSettings/Private/Settings/CaptureManagerSettings.cpp
--- a/Engine/Plugins/VirtualProduction/CaptureManager/CaptureManagerApp/Source/CaptureManagerSettings/Private/Settings/CaptureManagerSettings.cpp
+++ b/Engine/Plugins/VirtualProduction/CaptureManager/CaptureManagerApp/Source/CaptureManagerSettings/Private/Settings/CaptureManagerSettings.cpp
@@ -56,12 +56,12 @@ void UCaptureManagerSettings::InitializeValuesIfNotSet()
 
 	if (DefaultWorkingDirectory.Path.IsEmpty())
 	{
-		DefaultWorkingDirectory.Path = FPaths::Combine(FPlatformProcess::UserTempDir(), Private::DefaultWorkingDirectory);
+		DefaultWorkingDirectory.Path = FPaths::Combine(FPlatformProcess::UserDir(), Private::DefaultWorkingDirectory);
 	}
 
 	if (DownloadDirectory.Path.IsEmpty())
 	{
-		DownloadDirectory.Path = FPaths::Combine(FPlatformProcess::UserTempDir(), Private::DefaultDownloadDirectory);
+		DownloadDirectory.Path = FPaths::Combine(FPlatformProcess::UserDir(), Private::DefaultDownloadDirectory);
 	}
 
 	if (DefaultUploadHostName.IsEmpty())
diff --git a/Engine/Plugins/VirtualProduction/CaptureManager/CaptureManagerEditor/Source/CaptureManagerEditorSettings/Private/Settings/CaptureManagerEditorSettings.cpp b/Engine/Plugins/VirtualProduction/CaptureManager/CaptureManagerEditor/Source/CaptureManagerEditorSettings/Private/Settings/CaptureManagerEditorSettings.cpp
--- a/Engine/Plugins/VirtualProduction/CaptureManager/CaptureManagerEditor/Source/CaptureManagerEditorSettings/Private/Settings/CaptureManagerEditorSettings.cpp
+++ b/Engine/Plugins/VirtualProduction/CaptureManager/CaptureManagerEditor/Source/CaptureManagerEditorSettings/Private/Settings/CaptureManagerEditorSettings.cpp
@@ -221,7 +221,7 @@ void UCaptureManagerEditorSettings::InitializeValuesIfNotSet()
 
 	if (MediaDirectory.Path.IsEmpty())
 	{
-		MediaDirectory.Path = FPaths::Combine(FPlatformProcess::UserTempDir(), Private::DefaultMediaDirectory);
+		MediaDirectory.Path = FPaths::Combine(FPlatformProcess::UserDir(), Private::DefaultMediaDirectory);
 	}
 
 	if (ImportDirectory.Path.IsEmpty())
