Crash on take archive device path selected

diff --git a/Engine/Plugins/VirtualProduction/CaptureManager/CaptureManagerApp/Source/CaptureManagerEditor/Private/UI/CaptureManagerPanelController.cpp b/Engine/Plugins/VirtualProduction/CaptureManager/CaptureManagerApp/Source/CaptureManagerEditor/Private/UI/CaptureManagerPanelController.cpp
--- a/Engine/Plugins/VirtualProduction/CaptureManager/CaptureManagerApp/Source/CaptureManagerEditor/Private/UI/CaptureManagerPanelController.cpp
+++ b/Engine/Plugins/VirtualProduction/CaptureManager/CaptureManagerApp/Source/CaptureManagerEditor/Private/UI/CaptureManagerPanelController.cpp
@@ -177,6 +177,8 @@ void FCaptureManagerPanelController::OnDeviceAdded(FGuid InDeviceId, ULiveLinkDe
 		UConnectionDelegate* Delegate = ILiveLinkDeviceCapability_Connection::Execute_GetConnectionDelegate(InDevice);
 		Delegate->ConnectionChanged.AddSP(AsShared(), &FCaptureManagerPanelController::OnReachableEvent, InDeviceId);
 	}
+
+	TakesView->CaptureDeviceAdded(InDevice);
 }
 
 void FCaptureManagerPanelController::OnDeviceRemoved(FGuid InDeviceId, ULiveLinkDevice* InDevice)
@@ -192,7 +194,7 @@ void FCaptureManagerPanelController::OnDeviceRemoved(FGuid InDeviceId, ULiveLink
 		Delegate->ConnectionChanged.RemoveAll(this);
 	}
 
-	TakesView->CaptureDeviceToBeRemoved(InDeviceId);
+	TakesView->CaptureDeviceToBeRemoved(InDevice);
 }
 
 void FCaptureManagerPanelController::OnReachableEvent(ELiveLinkDeviceConnectionStatus InStatus, FGuid InDeviceId)
diff --git a/Engine/Plugins/VirtualProduction/CaptureManager/CaptureManagerApp/Source/CaptureManagerEditor/Private/UI/TakesView.cpp b/Engine/Plugins/VirtualProduction/CaptureManager/CaptureManagerApp/Source/CaptureManagerEditor/Private/UI/TakesView.cpp
--- a/Engine/Plugins/VirtualProduction/CaptureManager/CaptureManagerApp/Source/CaptureManagerEditor/Private/UI/TakesView.cpp
+++ b/Engine/Plugins/VirtualProduction/CaptureManager/CaptureManagerApp/Source/CaptureManagerEditor/Private/UI/TakesView.cpp
@@ -596,82 +596,6 @@ void FTakesView::CaptureDeviceStarted(FGuid InCaptureDeviceId)
 				return;
 			}
 
-			FString PackageName = GetTransientPackage()->GetPathName() / OuterPackageName;
-			UPackage* Package = CreatePackage(*PackageName);
-
-			FString CaptureDeviceName = Device->GetName();
-			FName Name(CaptureDeviceName);
-
-			TScriptInterface<ILiveLinkDeviceCapability_Ingest> CapabilityInterface(Device);
-
-			CapabilityInterface->SubscribeToEvent(FString(FIngestCapability_TakeAddedEvent::Name),
-				FCaptureEventHandler([this, InCaptureDeviceId](TSharedPtr<const FCaptureEvent> InTakeAddedEvent)
-					{
-						check(IsInGameThread());
-
-						TSharedPtr<const FIngestCapability_TakeAddedEvent> Event =
-							StaticCastSharedPtr<const FIngestCapability_TakeAddedEvent>(InTakeAddedEvent);
-
-						TArray<FTakeId> TakeIds = { Event->TakeId };
-						TArray<TStrongObjectPtr<UTakeVirtualAsset>> NewTakeObject = MakeTakeObjects(InCaptureDeviceId, { TakeIds });
-
-						if (TArray<TStrongObjectPtr<UTakeVirtualAsset>>* TakeObjects = CaptureDeviceTakesMap.Find(InCaptureDeviceId))
-						{
-							TakeObjects->Append(NewTakeObject.GetData(), NewTakeObject.Num());
-
-							RefreshTakesViewDelegate.ExecuteIfBound(true);
-						}
-					}));
-
-			CapabilityInterface->SubscribeToEvent(FString(FIngestCapability_TakeUpdatedEvent::Name),
-				FCaptureEventHandler([this, InCaptureDeviceId](TSharedPtr<const FCaptureEvent> InTakeUpdatedEvent)
-					{
-						check(IsInGameThread());
-
-						TSharedPtr<const FIngestCapability_TakeUpdatedEvent> Event =
-							StaticCastSharedPtr<const FIngestCapability_TakeUpdatedEvent>(InTakeUpdatedEvent);
-
-						TArray<FTakeId> TakeIds = { Event->TakeId };
-						TArray<TStrongObjectPtr<UTakeVirtualAsset>> NewTakeObject = MakeTakeObjects(InCaptureDeviceId, { TakeIds });
-						check(NewTakeObject.Num() == 1);
-
-						if (TArray<TStrongObjectPtr<UTakeVirtualAsset>>* TakeObjects = CaptureDeviceTakesMap.Find(InCaptureDeviceId))
-						{
-							for (TStrongObjectPtr<UTakeVirtualAsset>& TakeVirtualAsset : *TakeObjects)
-							{
-								if (TakeVirtualAsset->TakeId == Event->TakeId)
-								{
-									TakeVirtualAsset = NewTakeObject[0];
-									break;
-								}
-							}
-
-							RefreshTakesViewDelegate.ExecuteIfBound(true);
-						}
-					}));
-
-			CapabilityInterface->SubscribeToEvent(FString(FIngestCapability_TakeRemovedEvent::Name),
-				FCaptureEventHandler([this, InCaptureDeviceId](TSharedPtr<const FCaptureEvent> InTakeRemovedEvent)
-					{
-						check(IsInGameThread());
-
-						TSharedPtr<const FIngestCapability_TakeRemovedEvent> Event =
-							StaticCastSharedPtr<const FIngestCapability_TakeRemovedEvent>(InTakeRemovedEvent);
-
-						if (TArray<TStrongObjectPtr<UTakeVirtualAsset>>* TakeObjects = CaptureDeviceTakesMap.Find(InCaptureDeviceId))
-						{
-							int32 NumRemoved = TakeObjects->RemoveAll([this, TakeId = Event->TakeId](const TStrongObjectPtr<UTakeVirtualAsset>& InObject)
-								{
-									return InObject->TakeId == TakeId;
-								});
-
-							if (NumRemoved > 0)
-							{
-								RefreshTakesViewDelegate.ExecuteIfBound(true);
-							}
-						}
-					}));
-
 			ResetTakesCache(InCaptureDeviceId);
 
 			FIngestUpdateTakeListCallback Callback = 
@@ -709,11 +633,117 @@ void FTakesView::CaptureDeviceStopped(FGuid InCaptureDeviceId)
 		});
 }
 
-void FTakesView::CaptureDeviceToBeRemoved(FGuid InCaptureDeviceId)
+void FTakesView::CaptureDeviceAdded(ULiveLinkDevice* InDevice)
 {
-	AsyncTask(ENamedThreads::GameThread, [this, InCaptureDeviceId]()
+	AsyncTask(ENamedThreads::GameThread, [this, Controller = PanelController.Pin(), Device = TStrongObjectPtr<ULiveLinkDevice>(InDevice)]()
+	{
+		if (!Controller)
 		{
-			CaptureDeviceTakesMap.Remove(InCaptureDeviceId);
+			return;
+		}
+
+		using namespace UE::CaptureManager;
+
+		if (!Device)
+		{
+			return;
+		}
+
+		if (!Device->Implements<ULiveLinkDeviceCapability_Ingest>())
+		{
+			return;
+		}
+
+		FString PackageName = GetTransientPackage()->GetPathName() / OuterPackageName;
+		UPackage* Package = FindPackage(nullptr, *PackageName);
+		if (!Package)
+		{
+			Package = CreatePackage(*PackageName);
+		}
+
+		TScriptInterface<ILiveLinkDeviceCapability_Ingest> CapabilityInterface(Device.Get());
+
+		FGuid CaptureDeviceId = Device->GetDeviceId();
+
+		CapabilityInterface->SubscribeToEvent(FString(FIngestCapability_TakeAddedEvent::Name),
+			FCaptureEventHandler([this, CaptureDeviceId](TSharedPtr<const FCaptureEvent> InTakeAddedEvent)
+				{
+					check(IsInGameThread());
+
+					TSharedPtr<const FIngestCapability_TakeAddedEvent> Event =
+						StaticCastSharedPtr<const FIngestCapability_TakeAddedEvent>(InTakeAddedEvent);
+
+					TArray<FTakeId> TakeIds = { Event->TakeId };
+					TArray<TStrongObjectPtr<UTakeVirtualAsset>> NewTakeObject = MakeTakeObjects(CaptureDeviceId, { TakeIds });
+
+					if (TArray<TStrongObjectPtr<UTakeVirtualAsset>>* TakeObjects = CaptureDeviceTakesMap.Find(CaptureDeviceId))
+					{
+						TakeObjects->Append(NewTakeObject.GetData(), NewTakeObject.Num());
+
+						RefreshTakesViewDelegate.ExecuteIfBound(true);
+					}
+				}));
+
+		CapabilityInterface->SubscribeToEvent(FString(FIngestCapability_TakeUpdatedEvent::Name),
+			FCaptureEventHandler([this, CaptureDeviceId](TSharedPtr<const FCaptureEvent> InTakeUpdatedEvent)
+				{
+					check(IsInGameThread());
+
+					TSharedPtr<const FIngestCapability_TakeUpdatedEvent> Event =
+						StaticCastSharedPtr<const FIngestCapability_TakeUpdatedEvent>(InTakeUpdatedEvent);
+
+					TArray<FTakeId> TakeIds = { Event->TakeId };
+					TArray<TStrongObjectPtr<UTakeVirtualAsset>> NewTakeObject = MakeTakeObjects(CaptureDeviceId, { TakeIds });
+					check(NewTakeObject.Num() == 1);
+
+					if (TArray<TStrongObjectPtr<UTakeVirtualAsset>>* TakeObjects = CaptureDeviceTakesMap.Find(CaptureDeviceId))
+					{
+						for (TStrongObjectPtr<UTakeVirtualAsset>& TakeVirtualAsset : *TakeObjects)
+						{
+							if (TakeVirtualAsset->TakeId == Event->TakeId)
+							{
+								TakeVirtualAsset = NewTakeObject[0];
+								break;
+							}
+						}
+
+						RefreshTakesViewDelegate.ExecuteIfBound(true);
+					}
+				}));
+
+		CapabilityInterface->SubscribeToEvent(FString(FIngestCapability_TakeRemovedEvent::Name),
+			FCaptureEventHandler([this, CaptureDeviceId](TSharedPtr<const FCaptureEvent> InTakeRemovedEvent)
+				{
+					check(IsInGameThread());
+
+					TSharedPtr<const FIngestCapability_TakeRemovedEvent> Event =
+						StaticCastSharedPtr<const FIngestCapability_TakeRemovedEvent>(InTakeRemovedEvent);
+
+					if (TArray<TStrongObjectPtr<UTakeVirtualAsset>>* TakeObjects = CaptureDeviceTakesMap.Find(CaptureDeviceId))
+					{
+						int32 NumRemoved = TakeObjects->RemoveAll([this, TakeId = Event->TakeId](const TStrongObjectPtr<UTakeVirtualAsset>& InObject)
+							{
+								return InObject->TakeId == TakeId;
+							});
+
+						if (NumRemoved > 0)
+						{
+							RefreshTakesViewDelegate.ExecuteIfBound(true);
+						}
+					}
+				}));
+	});
+}
+
+void FTakesView::CaptureDeviceToBeRemoved(ULiveLinkDevice* InDevice)
+{
+	AsyncTask(ENamedThreads::GameThread, [this, Controller = PanelController.Pin(), Device = TStrongObjectPtr<ULiveLinkDevice>(InDevice)]()
+		{
+			TScriptInterface<ILiveLinkDeviceCapability_Ingest> CapabilityInterface(Device.Get());
+
+			CapabilityInterface->UnsubscribeAll();
+
+			CaptureDeviceTakesMap.Remove(Device->GetDeviceId());
 			RefreshTakesViewDelegate.ExecuteIfBound(true);
 		});
 }
@@ -757,7 +787,11 @@ TArray<TStrongObjectPtr<UTakeVirtualAsset>> FTakesView::MakeTakeObjects(FGuid In
 	}
 
 	FString PackageName = GetTransientPackage()->GetPathName() / OuterPackageName / CaptureDeviceName;
-	UPackage* Package = CreatePackage(*PackageName);
+	UPackage* Package = FindPackage(nullptr, *PackageName);
+	if (!Package)
+	{
+		Package = CreatePackage(*PackageName);
+	}
 
 	for (const FTakeId& TakeId : TakeIds)
 	{
@@ -827,7 +861,7 @@ FName FTakesView::MakeUniqueTakeName(FTakeMetadata TakeInfo, UPackage* Package)
 	UObject* FoundObject = StaticFindObjectFastInternal(nullptr, Package, UniqueName);
 
 	int32 UniquePartIncrement = 0;
-	while (FoundObject)
+	while (IsValid(FoundObject))
 	{
 		FName MaybeUniqueName(FString::Format(TEXT("{0}_{1}_{2}"), { TakeInfo.Slate, TakeInfo.TakeNumber, ++UniquePartIncrement }));
 		FoundObject = StaticFindObjectFastInternal(nullptr, Package, MaybeUniqueName);
diff --git a/Engine/Plugins/VirtualProduction/CaptureManager/CaptureManagerApp/Source/CaptureManagerEditor/Private/UI/TakesView.h b/Engine/Plugins/VirtualProduction/CaptureManager/CaptureManagerApp/Source/CaptureManagerEditor/Private/UI/TakesView.h
--- a/Engine/Plugins/VirtualProduction/CaptureManager/CaptureManagerApp/Source/CaptureManagerEditor/Private/UI/TakesView.h
+++ b/Engine/Plugins/VirtualProduction/CaptureManager/CaptureManagerApp/Source/CaptureManagerEditor/Private/UI/TakesView.h
@@ -83,7 +83,8 @@ public:
 
 	void CaptureDeviceStarted(FGuid InCaptureDeviceId);
 	void CaptureDeviceStopped(FGuid InCaptureDeviceId);
-	void CaptureDeviceToBeRemoved(FGuid InCaptureDeviceId);
+	void CaptureDeviceAdded(ULiveLinkDevice* InDevice);
+	void CaptureDeviceToBeRemoved(ULiveLinkDevice* InDevice);
 
 	TSharedPtr<STakesView> TakesTileView;
 private:
diff --git a/Engine/Plugins/VirtualProduction/CaptureManager/CaptureManagerApp/Source/LiveLinkCapabilities/Public/Ingest/IngestCapability_Events.h b/Engine/Plugins/VirtualProduction/CaptureManager/CaptureManagerApp/Source/LiveLinkCapabilities/Public/Ingest/IngestCapability_Events.h
--- a/Engine/Plugins/VirtualProduction/CaptureManager/CaptureManagerApp/Source/LiveLinkCapabilities/Public/Ingest/IngestCapability_Events.h
+++ b/Engine/Plugins/VirtualProduction/CaptureManager/CaptureManagerApp/Source/LiveLinkCapabilities/Public/Ingest/IngestCapability_Events.h
@@ -46,4 +46,14 @@ struct LIVELINKCAPABILITIES_API FIngestCapability_TakeRemovedEvent : public FIng
 		: FIngestCapability_TakeBaseEvent(InTakeId, FString(Name))
 	{
 	}
+};
+
+struct LIVELINKCAPABILITIES_API FIngestCapability_TakeListResetEvent : public UE::CaptureManager::FCaptureEvent
+{
+	static constexpr FStringView Name = TEXT("IngestCapability_TakeListReset");
+
+	explicit FIngestCapability_TakeListResetEvent(FString InName)
+		: FCaptureEvent(MoveTemp(InName))
+	{
+	}
 };
\ No newline at end of file
diff --git a/Engine/Plugins/VirtualProduction/CaptureManager/CaptureManagerDevices/Source/FilesystemLiveLinkDevice/Private/FilesystemLiveLinkDevice.cpp b/Engine/Plugins/VirtualProduction/CaptureManager/CaptureManagerDevices/Source/FilesystemLiveLinkDevice/Private/FilesystemLiveLinkDevice.cpp
--- a/Engine/Plugins/VirtualProduction/CaptureManager/CaptureManagerDevices/Source/FilesystemLiveLinkDevice/Private/FilesystemLiveLinkDevice.cpp
+++ b/Engine/Plugins/VirtualProduction/CaptureManager/CaptureManagerDevices/Source/FilesystemLiveLinkDevice/Private/FilesystemLiveLinkDevice.cpp
@@ -46,6 +46,8 @@ void UFilesystemLiveLinkDevice::OnSettingChanged(const FPropertyChangedEvent& In
 
 	if (GET_MEMBER_NAME_CHECKED(UFilesystemLiveLinkDeviceSettings, TakeDirectory) == InPropertyChangedEvent.GetMemberPropertyName())
 	{
+		SetConnectionStatus(ELiveLinkDeviceConnectionStatus::Disconnected);
+
 		FString Path = GetSettings()->TakeDirectory.Path;
 		if (!Path.IsEmpty() && FPaths::DirectoryExists(Path))
 		{
diff --git a/Engine/Plugins/VirtualProduction/CaptureManager/CaptureManagerDevices/Source/FilesystemLiveLinkDevice/Private/Tests/GenericLiveLinkDeviceTests.cpp b/Engine/Plugins/VirtualProduction/CaptureManager/CaptureManagerDevices/Source/FilesystemLiveLinkDevice/Private/Tests/GenericLiveLinkDeviceTests.cpp
deleted file mode 100644
diff --git a/Engine/Plugins/VirtualProduction/CaptureManager/CaptureManagerDevices/Source/MonoVideoLiveLinkDevice/Private/MonoVideoLiveLinkDevice.cpp b/Engine/Plugins/VirtualProduction/CaptureManager/CaptureManagerDevices/Source/MonoVideoLiveLinkDevice/Private/MonoVideoLiveLinkDevice.cpp
--- a/Engine/Plugins/VirtualProduction/CaptureManager/CaptureManagerDevices/Source/MonoVideoLiveLinkDevice/Private/MonoVideoLiveLinkDevice.cpp
+++ b/Engine/Plugins/VirtualProduction/CaptureManager/CaptureManagerDevices/Source/MonoVideoLiveLinkDevice/Private/MonoVideoLiveLinkDevice.cpp
@@ -77,6 +77,8 @@ void UMonoVideoLiveLinkDevice::OnSettingChanged(const FPropertyChangedEvent& InP
 	if (GET_MEMBER_NAME_CHECKED(UMonoVideoLiveLinkDeviceSettings, TakeDirectory) == PropertyName
 		|| GET_MEMBER_NAME_CHECKED(UMonoVideoLiveLinkDeviceSettings, VideoDiscoveryExpression) == PropertyName)
 	{
+		SetConnectionStatus(ELiveLinkDeviceConnectionStatus::Disconnected);
+
 		FString Path = GetSettings()->TakeDirectory.Path;
 		if (!Path.IsEmpty() && FPaths::DirectoryExists(Path))
 		{
diff --git a/Engine/Plugins/VirtualProduction/CaptureManager/CaptureManagerDevices/Source/StereoVideoLiveLinkDevice/Private/StereoVideoLiveLinkDevice.cpp b/Engine/Plugins/VirtualProduction/CaptureManager/CaptureManagerDevices/Source/StereoVideoLiveLinkDevice/Private/StereoVideoLiveLinkDevice.cpp
--- a/Engine/Plugins/VirtualProduction/CaptureManager/CaptureManagerDevices/Source/StereoVideoLiveLinkDevice/Private/StereoVideoLiveLinkDevice.cpp
+++ b/Engine/Plugins/VirtualProduction/CaptureManager/CaptureManagerDevices/Source/StereoVideoLiveLinkDevice/Private/StereoVideoLiveLinkDevice.cpp
@@ -89,6 +89,8 @@ void UStereoVideoLiveLinkDevice::OnSettingChanged(const FPropertyChangedEvent& I
 		|| GET_MEMBER_NAME_CHECKED(UStereoVideoLiveLinkDeviceSettings, VideoDiscoveryExpression) == PropertyName
 		|| GET_MEMBER_NAME_CHECKED(UStereoVideoLiveLinkDeviceSettings, AudioDiscoveryExpression) == PropertyName)
 	{
+		SetConnectionStatus(ELiveLinkDeviceConnectionStatus::Disconnected);
+
 		FString Path = GetSettings()->TakeDirectory.Path;
 		if (!Path.IsEmpty() && FPaths::DirectoryExists(Path))
 		{
