[CaptureManager] Updated launch ingest server icon

diff --git a/Engine/Plugins/VirtualProduction/CaptureManager/CaptureManagerApp/Source/CaptureManagerEditor/CaptureManagerEditor.Build.cs b/Engine/Plugins/VirtualProduction/CaptureManager/CaptureManagerApp/Source/CaptureManagerEditor/CaptureManagerEditor.Build.cs
--- a/Engine/Plugins/VirtualProduction/CaptureManager/CaptureManagerApp/Source/CaptureManagerEditor/CaptureManagerEditor.Build.cs
+++ b/Engine/Plugins/VirtualProduction/CaptureManager/CaptureManagerApp/Source/CaptureManagerEditor/CaptureManagerEditor.Build.cs
@@ -31,6 +31,7 @@ public class CaptureManagerEditor : ModuleRules
 			"Engine",
 			"CaptureUtils",
 			"CaptureManagerTakeMetadata",
+			"CaptureManagerStyle",
 			"OutputLog",			
 			"CaptureDataCore",
 			"CaptureManagerSettings",
diff --git a/Engine/Plugins/VirtualProduction/CaptureManager/CaptureManagerApp/Source/CaptureManagerEditor/Private/CaptureManagerEditorModule.cpp b/Engine/Plugins/VirtualProduction/CaptureManager/CaptureManagerApp/Source/CaptureManagerEditor/Private/CaptureManagerEditorModule.cpp
--- a/Engine/Plugins/VirtualProduction/CaptureManager/CaptureManagerApp/Source/CaptureManagerEditor/Private/CaptureManagerEditorModule.cpp
+++ b/Engine/Plugins/VirtualProduction/CaptureManager/CaptureManagerApp/Source/CaptureManagerEditor/Private/CaptureManagerEditorModule.cpp
@@ -10,7 +10,6 @@
 #include "UI/TakeVirtualAsset.h"
 
 #include "Styling/SlateStyleRegistry.h"
-#include "UI/Style/CaptureManagerTabStyle.h"
 #include "UI/IngestJobSettingsCustomization.h"
 
 #include "IngestManagement/UIngestJobSettings.h"
@@ -24,15 +23,11 @@ void FCaptureManagerEditorModule::StartupModule()
 	IModularFeatures::Get().RegisterModularFeature(ILiveLinkHubApplicationModeFactory::ModularFeatureName, this);
 	UThumbnailManager::Get().RegisterCustomRenderer(UTakeVirtualAsset::StaticClass(), UTakeThumbnailRenderer::StaticClass());
 
-	FSlateStyleRegistry::RegisterSlateStyle(FCaptureManagerTabStyle::Get());
-
 	FPropertyEditorModule& PropertyEditorModule = FModuleManager::GetModuleChecked<FPropertyEditorModule>("PropertyEditor");
 	PropertyEditorModule.RegisterCustomClassLayout(
 		UIngestJobSettingsBase::StaticClass()->GetFName(),
 		FOnGetDetailCustomizationInstance::CreateStatic(&FIngestJobSettingsCustomization::MakeInstance)
 	);
-
-	FCaptureManagerTabStyle::ReloadTextures();
 }
 
 void FCaptureManagerEditorModule::ShutdownModule()
@@ -42,8 +37,6 @@ void FCaptureManagerEditorModule::ShutdownModule()
 		UThumbnailManager::Get().UnregisterCustomRenderer(UTakeThumbnailRenderer::StaticClass());
 	}
 
-	FSlateStyleRegistry::UnRegisterSlateStyle(FCaptureManagerTabStyle::Get());
-
 	FPropertyEditorModule& PropertyEditorModule = FModuleManager::GetModuleChecked<FPropertyEditorModule>("PropertyEditor");
 	PropertyEditorModule.UnregisterCustomClassLayout(UIngestJobSettingsBase::StaticClass()->GetFName());
 }
diff --git a/Engine/Plugins/VirtualProduction/CaptureManager/CaptureManagerApp/Source/CaptureManagerEditor/Private/UI/LiveLinkHubCaptureManagerMode.h b/Engine/Plugins/VirtualProduction/CaptureManager/CaptureManagerApp/Source/CaptureManagerEditor/Private/UI/LiveLinkHubCaptureManagerMode.h
--- a/Engine/Plugins/VirtualProduction/CaptureManager/CaptureManagerApp/Source/CaptureManagerEditor/Private/UI/LiveLinkHubCaptureManagerMode.h
+++ b/Engine/Plugins/VirtualProduction/CaptureManager/CaptureManagerApp/Source/CaptureManagerEditor/Private/UI/LiveLinkHubCaptureManagerMode.h
@@ -9,13 +9,13 @@
 #include "Framework/Docking/TabManager.h"
 #include "Templates/SharedPointer.h"
 
-#include "Style/CaptureManagerTabStyle.h"
-
 #include "Async/CaptureTimerManager.h"
 
 #include "IMessageContext.h"
 #include "WorkflowOrientedApp/WorkflowTabManager.h"
 
+#include "CaptureManagerStyle.h"
+
 
 class FCommunicationManager;
 
@@ -31,7 +31,7 @@ public:
 	//~ Begin FLiveLinkHubApplicationMode interface
 	virtual FSlateIcon GetModeIcon() const override
 	{
-		return FSlateIcon(FCaptureManagerTabStyle::Get().GetStyleSetName(), TEXT("CaptureManager.Tabs.CaptureManager"));
+		return FSlateIcon(FCaptureManagerStyle::Get().GetStyleSetName(), TEXT("CaptureManagerIcon"));
 	}
 
 	virtual TArray<TSharedRef<SWidget>> GetStatusBarWidgets_Impl() override;
diff --git a/Engine/Plugins/VirtualProduction/CaptureManager/CaptureManagerApp/Source/CaptureManagerEditor/Private/UI/Style/CaptureManagerTabStyle.cpp b/Engine/Plugins/VirtualProduction/CaptureManager/CaptureManagerApp/Source/CaptureManagerEditor/Private/UI/Style/CaptureManagerTabStyle.cpp
deleted file mode 100644
diff --git a/Engine/Plugins/VirtualProduction/CaptureManager/CaptureManagerApp/Source/CaptureManagerEditor/Private/UI/Style/CaptureManagerTabStyle.h b/Engine/Plugins/VirtualProduction/CaptureManager/CaptureManagerApp/Source/CaptureManagerEditor/Private/UI/Style/CaptureManagerTabStyle.h
deleted file mode 100644
diff --git a/Engine/Plugins/VirtualProduction/CaptureManager/CaptureManagerCore/CaptureManagerCore.uplugin b/Engine/Plugins/VirtualProduction/CaptureManager/CaptureManagerCore/CaptureManagerCore.uplugin
--- a/Engine/Plugins/VirtualProduction/CaptureManager/CaptureManagerCore/CaptureManagerCore.uplugin
+++ b/Engine/Plugins/VirtualProduction/CaptureManager/CaptureManagerCore/CaptureManagerCore.uplugin
@@ -37,6 +37,11 @@
 			"Type": "Editor",
 			"LoadingPhase": "Default"
 		},
+		{
+			"Name": "CaptureManagerStyle",
+			"Type": "Editor",
+			"LoadingPhase": "Default"
+		},
 		{
 			"Name": "CaptureManagerTakeMetadata",
 			"Type": "Editor",
diff --git a/Engine/Plugins/VirtualProduction/CaptureManager/CaptureManagerCore/Source/CaptureManagerStyle/CaptureManagerStyle.Build.cs b/Engine/Plugins/VirtualProduction/CaptureManager/CaptureManagerCore/Source/CaptureManagerStyle/CaptureManagerStyle.Build.cs
new file mode 100644
--- /dev/null
+++ b/Engine/Plugins/VirtualProduction/CaptureManager/CaptureManagerCore/Source/CaptureManagerStyle/CaptureManagerStyle.Build.cs
@@ -0,0 +1,25 @@
+// Copyright Epic Games, Inc. All Rights Reserved.
+
+using UnrealBuildTool;
+
+public class CaptureManagerStyle : ModuleRules
+{
+	public CaptureManagerStyle(ReadOnlyTargetRules Target) : base(Target)
+	{
+		PCHUsage = PCHUsageMode.UseExplicitOrSharedPCHs;
+
+		PublicDependencyModuleNames.AddRange(new string[]
+		{
+			"Core",
+			"EditorStyle"
+		});
+
+		PrivateDependencyModuleNames.AddRange(new string[]
+		{
+			"Projects",
+			"Slate",
+			"SlateCore",
+			"UnrealEd"
+		});
+	}
+}
\ No newline at end of file
diff --git a/Engine/Plugins/VirtualProduction/CaptureManager/CaptureManagerCore/Source/CaptureManagerStyle/Private/CaptureManagerStyle.cpp b/Engine/Plugins/VirtualProduction/CaptureManager/CaptureManagerCore/Source/CaptureManagerStyle/Private/CaptureManagerStyle.cpp
new file mode 100644
--- /dev/null
+++ b/Engine/Plugins/VirtualProduction/CaptureManager/CaptureManagerCore/Source/CaptureManagerStyle/Private/CaptureManagerStyle.cpp
@@ -0,0 +1,39 @@
+// Copyright Epic Games, Inc. All Rights Reserved.
+
+#include "CaptureManagerStyle.h"
+#include "Framework/Application/SlateApplication.h"
+#include "Slate/SlateGameResources.h"
+#include "Interfaces/IPluginManager.h"
+#include "Styling/SlateStyleMacros.h"
+
+FName FCaptureManagerStyle::StyleName = TEXT("CaptureManagerStyle");
+
+FCaptureManagerStyle::FCaptureManagerStyle()
+	: FSlateStyleSet(FCaptureManagerStyle::StyleName)
+{
+	const FVector2D Icon16x16(16.0f, 16.0f);
+
+	SetContentRoot(IPluginManager::Get().FindPlugin(UE_PLUGIN_NAME)->GetContentDir());
+
+	Set("CaptureManagerIcon", new IMAGE_BRUSH_SVG("Icons/CaptureManager_16", Icon16x16));
+}
+
+const FName& FCaptureManagerStyle::GetStyleSetName() const
+{
+	return StyleName;
+}
+
+const FCaptureManagerStyle& FCaptureManagerStyle::Get()
+{
+	static FCaptureManagerStyle StyleInstance;
+
+	return StyleInstance;
+}
+
+void FCaptureManagerStyle::ReloadTextures()
+{
+	if (FSlateApplication::IsInitialized())
+	{
+		FSlateApplication::Get().GetRenderer()->ReloadTextureResources();
+	}
+}
diff --git a/Engine/Plugins/VirtualProduction/CaptureManager/CaptureManagerCore/Source/CaptureManagerStyle/Private/CaptureManagerStyleModule.cpp b/Engine/Plugins/VirtualProduction/CaptureManager/CaptureManagerCore/Source/CaptureManagerStyle/Private/CaptureManagerStyleModule.cpp
new file mode 100644
--- /dev/null
+++ b/Engine/Plugins/VirtualProduction/CaptureManager/CaptureManagerCore/Source/CaptureManagerStyle/Private/CaptureManagerStyleModule.cpp
@@ -0,0 +1,20 @@
+// Copyright Epic Games, Inc. All Rights Reserved.
+
+#include "CaptureManagerStyleModule.h"
+
+#include "Styling/SlateStyleRegistry.h"
+#include "CaptureManagerStyle.h"
+
+
+IMPLEMENT_MODULE(FCaptureManagerStyleModule, CaptureManagerStyle)
+
+void FCaptureManagerStyleModule::StartupModule()
+{
+	FSlateStyleRegistry::RegisterSlateStyle(FCaptureManagerStyle::Get());
+	FCaptureManagerStyle::ReloadTextures();
+}
+
+void FCaptureManagerStyleModule::ShutdownModule()
+{
+	FSlateStyleRegistry::UnRegisterSlateStyle(FCaptureManagerStyle::Get());
+}
diff --git a/Engine/Plugins/VirtualProduction/CaptureManager/CaptureManagerCore/Source/CaptureManagerStyle/Private/CaptureManagerStyleModule.h b/Engine/Plugins/VirtualProduction/CaptureManager/CaptureManagerCore/Source/CaptureManagerStyle/Private/CaptureManagerStyleModule.h
new file mode 100644
--- /dev/null
+++ b/Engine/Plugins/VirtualProduction/CaptureManager/CaptureManagerCore/Source/CaptureManagerStyle/Private/CaptureManagerStyleModule.h
@@ -0,0 +1,12 @@
+// Copyright Epic Games, Inc. All Rights Reserved.
+
+#pragma once
+
+#include "Modules/ModuleManager.h"
+
+class FCaptureManagerStyleModule : public IModuleInterface
+{
+public:
+	virtual void StartupModule() override;
+	virtual void ShutdownModule() override;
+};
diff --git a/Engine/Plugins/VirtualProduction/CaptureManager/CaptureManagerCore/Source/CaptureManagerStyle/Public/CaptureManagerStyle.h b/Engine/Plugins/VirtualProduction/CaptureManager/CaptureManagerCore/Source/CaptureManagerStyle/Public/CaptureManagerStyle.h
new file mode 100644
--- /dev/null
+++ b/Engine/Plugins/VirtualProduction/CaptureManager/CaptureManagerCore/Source/CaptureManagerStyle/Public/CaptureManagerStyle.h
@@ -0,0 +1,23 @@
+// Copyright Epic Games, Inc. All Rights Reserved.
+
+#pragma once
+
+#include "CoreMinimal.h"
+#include "Styling/SlateStyle.h"
+
+/** A style set for Capture Manager */
+class CAPTUREMANAGERSTYLE_API FCaptureManagerStyle : public FSlateStyleSet
+{
+public:
+
+	virtual const FName& GetStyleSetName() const override;
+	static const FCaptureManagerStyle& Get();
+
+	static void ReloadTextures();
+
+private:
+
+	FCaptureManagerStyle();
+
+	static FName StyleName;
+};
\ No newline at end of file
diff --git a/Engine/Plugins/VirtualProduction/CaptureManager/CaptureManagerEditor/Source/CaptureManagerEditorSettings/CaptureManagerEditorSettings.Build.cs b/Engine/Plugins/VirtualProduction/CaptureManager/CaptureManagerEditor/Source/CaptureManagerEditorSettings/CaptureManagerEditorSettings.Build.cs
--- a/Engine/Plugins/VirtualProduction/CaptureManager/CaptureManagerEditor/Source/CaptureManagerEditorSettings/CaptureManagerEditorSettings.Build.cs
+++ b/Engine/Plugins/VirtualProduction/CaptureManager/CaptureManagerEditor/Source/CaptureManagerEditorSettings/CaptureManagerEditorSettings.Build.cs
@@ -21,6 +21,7 @@ public class CaptureManagerEditorSettings : ModuleRules
 		{
 			"Engine",
 			"DeveloperSettings",
+			"CaptureManagerStyle",
 			"ContentBrowserData",
 			"LiveLinkHubCaptureMessaging",
 			"LiveLinkHubExportServer",
diff --git a/Engine/Plugins/VirtualProduction/CaptureManager/CaptureManagerEditor/Source/CaptureManagerEditorSettings/Private/CaptureManagerEditorSettingsModule.h b/Engine/Plugins/VirtualProduction/CaptureManager/CaptureManagerEditor/Source/CaptureManagerEditorSettings/Private/CaptureManagerEditorSettingsModule.h
--- a/Engine/Plugins/VirtualProduction/CaptureManager/CaptureManagerEditor/Source/CaptureManagerEditorSettings/Private/CaptureManagerEditorSettingsModule.h
+++ b/Engine/Plugins/VirtualProduction/CaptureManager/CaptureManagerEditor/Source/CaptureManagerEditorSettings/Private/CaptureManagerEditorSettingsModule.h
@@ -4,6 +4,8 @@
 
 #include "Modules/ModuleManager.h"
 
+class FSlateStyleSet;
+
 class FCaptureManagerEditorSettingsModule : public IModuleInterface
 {
 public:
@@ -14,4 +16,6 @@ private:
 
 	void PostEngineInit();
 	void EnginePreExit();
+
+	TSharedPtr<FSlateStyleSet> StyleSet;
 };
diff --git a/Engine/Plugins/VirtualProduction/CaptureManager/CaptureManagerEditor/Source/CaptureManagerEditorSettings/Private/Settings/CaptureManagerEditorSettings.cpp b/Engine/Plugins/VirtualProduction/CaptureManager/CaptureManagerEditor/Source/CaptureManagerEditorSettings/Private/Settings/CaptureManagerEditorSettings.cpp
--- a/Engine/Plugins/VirtualProduction/CaptureManager/CaptureManagerEditor/Source/CaptureManagerEditorSettings/Private/Settings/CaptureManagerEditorSettings.cpp
+++ b/Engine/Plugins/VirtualProduction/CaptureManager/CaptureManagerEditor/Source/CaptureManagerEditorSettings/Private/Settings/CaptureManagerEditorSettings.cpp
@@ -17,7 +17,8 @@
 #include "Framework/Notifications/NotificationManager.h"
 #include "Widgets/Notifications/SNotificationList.h"
 #include "Styling/SlateTypes.h"
-#include "Styling/SlateStyle.h"
+#include "CaptureManagerStyle.h"
+
 
 #define LOCTEXT_NAMESPACE "CaptureManagerEditorSettings"
 
@@ -138,7 +139,7 @@ void UCaptureManagerEditorSettings::Initialize()
 		Section.AddMenuEntry("IngestServer",
 			LOCTEXT("IngestServerLabel", "Ingest Server"),
 			LOCTEXT("IngestServerTooltip", "Launch the Capture Manager Ingest Server."),
-			FSlateIcon("LiveLinkStyle", "LiveLinkClient.Common.Icon.Small"),
+			FSlateIcon(FCaptureManagerStyle::Get().GetStyleSetName(), "CaptureManagerIcon"),
 			FUIAction(FExecuteAction::CreateLambda([this, &LiveLinkHubExportServerModule]() {
 				if (!LiveLinkHubExportServerModule.IsExportServerRunning())
 				{
