Take UObject memory management

diff --git a/Engine/Plugins/VirtualProduction/CaptureManager/CaptureManagerApp/Source/CaptureManagerEditor/Private/UI/CaptureManagerPanelController.cpp b/Engine/Plugins/VirtualProduction/CaptureManager/CaptureManagerApp/Source/CaptureManagerEditor/Private/UI/CaptureManagerPanelController.cpp
--- a/Engine/Plugins/VirtualProduction/CaptureManager/CaptureManagerApp/Source/CaptureManagerEditor/Private/UI/CaptureManagerPanelController.cpp
+++ b/Engine/Plugins/VirtualProduction/CaptureManager/CaptureManagerApp/Source/CaptureManagerEditor/Private/UI/CaptureManagerPanelController.cpp
@@ -165,6 +165,27 @@ TObjectPtr<ULiveLinkDevice> FCaptureManagerPanelController::GetCaptureDevice(FGu
 	return nullptr;
 }
 
+TArray<TObjectPtr<ULiveLinkDevice>> FCaptureManagerPanelController::GetCaptureDevices() const
+{
+	ULiveLinkDeviceSubsystem* Subsystem = GEngine->GetEngineSubsystem<ULiveLinkDeviceSubsystem>();
+
+	TArray<TObjectPtr<ULiveLinkDevice>> CaptureDevices;
+
+	for (const TPair<FGuid, TObjectPtr<ULiveLinkDevice>>& DevicePair : Subsystem->GetDeviceMap())
+	{
+		TObjectPtr<ULiveLinkDevice> Device = DevicePair.Value;
+
+		if (!Device->Implements<ULiveLinkDeviceCapability_Ingest>())
+		{
+			continue;
+		}
+		
+		CaptureDevices.Add(MoveTemp(Device));
+	}
+
+	return CaptureDevices;
+}
+
 void FCaptureManagerPanelController::OnDeviceAdded(FGuid InDeviceId, ULiveLinkDevice* InDevice)
 {
 	if (!InDevice->Implements<ULiveLinkDeviceCapability_Ingest>())
diff --git a/Engine/Plugins/VirtualProduction/CaptureManager/CaptureManagerApp/Source/CaptureManagerEditor/Private/UI/CaptureManagerPanelController.h b/Engine/Plugins/VirtualProduction/CaptureManager/CaptureManagerApp/Source/CaptureManagerEditor/Private/UI/CaptureManagerPanelController.h
--- a/Engine/Plugins/VirtualProduction/CaptureManager/CaptureManagerApp/Source/CaptureManagerEditor/Private/UI/CaptureManagerPanelController.h
+++ b/Engine/Plugins/VirtualProduction/CaptureManager/CaptureManagerApp/Source/CaptureManagerEditor/Private/UI/CaptureManagerPanelController.h
@@ -53,6 +53,7 @@ public:
 	TSharedRef<UE::CaptureManager::FIngestJobSettingsManager> GetIngestJobSettingsManager() const;
 
 	TObjectPtr<ULiveLinkDevice> GetCaptureDevice(FGuid InDeviceId) const;
+	TArray<TObjectPtr<ULiveLinkDevice>> GetCaptureDevices() const;
 
 private:
 
diff --git a/Engine/Plugins/VirtualProduction/CaptureManager/CaptureManagerApp/Source/CaptureManagerEditor/Private/UI/TUniqueObjectPtr.h b/Engine/Plugins/VirtualProduction/CaptureManager/CaptureManagerApp/Source/CaptureManagerEditor/Private/UI/TUniqueObjectPtr.h
new file mode 100644
--- /dev/null
+++ b/Engine/Plugins/VirtualProduction/CaptureManager/CaptureManagerApp/Source/CaptureManagerEditor/Private/UI/TUniqueObjectPtr.h
@@ -0,0 +1,81 @@
+// Copyright Epic Games, Inc. All Rights Reserved.
+
+#pragma once
+
+#include <type_traits>
+
+#include "UObject/Object.h"
+#include "UObject/StrongObjectPtrTemplates.h"
+
+template<class T>
+class TUniqueObjectPtr
+{
+public:
+
+	static_assert(std::is_base_of_v<UObject, T>, "Provided pointer has to be derived from UObject");
+
+	TUniqueObjectPtr()
+		: Object(nullptr)
+	{
+	}
+
+	TUniqueObjectPtr(T* InObject)
+		: Object(InObject)
+	{
+	}
+
+	~TUniqueObjectPtr()
+	{
+		if (IsValid() && Object.IsValid())
+		{
+			Object->MarkAsGarbage();
+
+			Object = nullptr;
+		}
+	}
+
+	TUniqueObjectPtr(const TUniqueObjectPtr<T>& InOther) = delete;
+	TUniqueObjectPtr<T>& operator=(const TUniqueObjectPtr<T>& InOther) = delete;
+
+	TUniqueObjectPtr(TUniqueObjectPtr<T>&& InOther) = default;
+	TUniqueObjectPtr<T>& operator=(TUniqueObjectPtr<T>&& InOther) = default;
+
+	const T& operator*() const
+	{
+		return *Object.Get();
+	}
+
+	T& operator*()
+	{
+		return *Object.Get();
+	}
+
+	const T* operator->() const
+	{
+		return Object.Get();
+	}
+
+	T* operator->()
+	{
+		return Object.Get();
+	}
+
+	T* Get()
+	{
+		return Object.Get();
+	}
+
+	const T* Get() const
+	{
+		return Object.Get();
+	}
+
+	bool IsValid() const
+	{
+		return Object != nullptr;
+	}
+
+private:
+
+	TStrongObjectPtr<T> Object;
+};
\ No newline at end of file
diff --git a/Engine/Plugins/VirtualProduction/CaptureManager/CaptureManagerApp/Source/CaptureManagerEditor/Private/UI/TakesView.cpp b/Engine/Plugins/VirtualProduction/CaptureManager/CaptureManagerApp/Source/CaptureManagerEditor/Private/UI/TakesView.cpp
--- a/Engine/Plugins/VirtualProduction/CaptureManager/CaptureManagerApp/Source/CaptureManagerEditor/Private/UI/TakesView.cpp
+++ b/Engine/Plugins/VirtualProduction/CaptureManager/CaptureManagerApp/Source/CaptureManagerEditor/Private/UI/TakesView.cpp
@@ -227,7 +227,7 @@ void STakesView::UpdateTakeListFinished()
 	RefreshButton->SetEnabled(true);
 }
 
-TArray<TStrongObjectPtr<UTakeVirtualAsset>> STakesView::GetSelectedTakeAssets() const
+TArray<TObjectPtr<UTakeVirtualAsset>> STakesView::GetSelectedTakeAssets() const
 {
 	check(GetCurrentSelectionDelegate);
 
@@ -238,12 +238,12 @@ TArray<TStrongObjectPtr<UTakeVirtualAsset>> STakesView::GetSelectedTakeAssets()
 
 	const TArray<FAssetData> SelectedAssetDatas = GetCurrentSelectionDelegate->Execute();
 
-	TArray<TStrongObjectPtr<UTakeVirtualAsset>> TakeAssets;
+	TArray<TObjectPtr<UTakeVirtualAsset>> TakeAssets;
 	TakeAssets.Reserve(SelectedAssetDatas.Num());
 
 	for (const FAssetData& AssetData : SelectedAssetDatas)
 	{
-		TStrongObjectPtr<UTakeVirtualAsset> TakeAsset(Cast<UTakeVirtualAsset>(AssetData.GetAsset()));
+		TObjectPtr<UTakeVirtualAsset> TakeAsset(Cast<UTakeVirtualAsset>(AssetData.GetAsset()));
 
 		if (TakeAsset)
 		{
@@ -267,14 +267,14 @@ FReply STakesView::OnAddToQueueButtonClicked()
 
 	if (IngestJobSettingsManager)
 	{
-		TArray<TStrongObjectPtr<UTakeVirtualAsset>> SelectedTakeAssets = GetSelectedTakeAssets();
+		TArray<TObjectPtr<UTakeVirtualAsset>> SelectedTakeAssets = GetSelectedTakeAssets();
 
 		// If the user has selected a lot of takes, rather than add them in a random order, we add them in the order of
 		// acquisition, just so that it is deterministic
 		SelectedTakeAssets.Sort(
-			[](const TStrongObjectPtr<const UTakeVirtualAsset>& LeftItem, const TStrongObjectPtr<const UTakeVirtualAsset>& RightItem)
+			[](const UTakeVirtualAsset& LeftItem, const UTakeVirtualAsset& RightItem)
 			{
-				return LeftItem->Metadata.DateTime.Get(FDateTime()) < RightItem->Metadata.DateTime.Get(FDateTime());
+				return LeftItem.Metadata.DateTime.Get(FDateTime()) < RightItem.Metadata.DateTime.Get(FDateTime());
 			}
 		);
 
@@ -282,7 +282,7 @@ FReply STakesView::OnAddToQueueButtonClicked()
 
 		if (SelectedPreset)
 		{
-			OnAddTakesToIngestQueue.ExecuteIfBound(SelectedTakeAssets, *SelectedPreset);
+			OnAddTakesToIngestQueue.ExecuteIfBound(MoveTemp(SelectedTakeAssets), *SelectedPreset);
 		}
 	}
 
@@ -322,23 +322,23 @@ TArray<FAssetData> FTakesView::CreateTakeAssetsDataForCaptureDevice(TOptional<FS
 {
 	using namespace UE::CaptureManager;
 
-	TArray<TStrongObjectPtr<UTakeVirtualAsset>> SelectedDevicesTakesArray;
+	TArray<TObjectPtr<UTakeVirtualAsset>> SelectedDevicesTakesArray;
 
 	TSharedPtr<FCaptureManagerPanelController> Controller = PanelController.Pin();
 	if (Controller)
 	{
-		for (const TPair<FGuid, TArray<TStrongObjectPtr<UTakeVirtualAsset>>>& DeviceIdTakeTuple : CaptureDeviceTakesMap)
+		for (TPair<FGuid, TArray<TUniqueObjectPtr<UTakeVirtualAsset>>>& DeviceIdTakeTuple : CaptureDeviceTakesMap)
 		{
 			FGuid CaptureDeviceId = DeviceIdTakeTuple.Key;
 			TObjectPtr<ULiveLinkDevice> Device = Controller->GetCaptureDevice(CaptureDeviceId);
 
 			if (!InCaptureDeviceName.IsSet() || Device.GetName() == InCaptureDeviceName.GetValue())
 			{
-				TArray<TStrongObjectPtr<UTakeVirtualAsset>> Takes;
+				TArray<TObjectPtr<UTakeVirtualAsset>> Takes;
 
-				for (const TStrongObjectPtr<UTakeVirtualAsset>& Take : DeviceIdTakeTuple.Value)
+				for (TUniqueObjectPtr<UTakeVirtualAsset>& Take : DeviceIdTakeTuple.Value)
 				{
-					Takes.Add(Take);
+					Takes.Add(TObjectPtr<UTakeVirtualAsset>(Take.Get()));
 				}
 
 				SelectedDevicesTakesArray.Append(MoveTemp(Takes));
@@ -349,10 +349,10 @@ TArray<FAssetData> FTakesView::CreateTakeAssetsDataForCaptureDevice(TOptional<FS
 	return CreateTakeAssetsData(SelectedDevicesTakesArray);
 }
 
-TArray<FAssetData> FTakesView::CreateTakeAssetsData(const TArray<TStrongObjectPtr<UTakeVirtualAsset>>& InSelectedCaptureDevicesTakesArray)
+TArray<FAssetData> FTakesView::CreateTakeAssetsData(const TArray<TObjectPtr<UTakeVirtualAsset>>& InSelectedCaptureDevicesTakesArray)
 {
 	TArray<FAssetData> TakeAssetsData;
-	for (TStrongObjectPtr<UTakeVirtualAsset> TakeItem : InSelectedCaptureDevicesTakesArray)
+	for (const TObjectPtr<UTakeVirtualAsset>& TakeItem : InSelectedCaptureDevicesTakesArray)
 	{
 		FAssetData TakeAssetData(TakeItem.Get());
 		TakeAssetData.AssetClassPath = FTopLevelAssetPath(*(TakeItem->GetPackage()->GetName()), *TakeItem->GetName());
@@ -385,7 +385,7 @@ EImageRotation FTakesView::ConvertOrientation(TOptional<FTakeMetadata::FVideo::E
 	}
 }
 
-void FTakesView::AddTakesToIngestQueue(const TArray<TStrongObjectPtr<UTakeVirtualAsset>>& InTakeAssets, const UIngestJobSettingsBase& InJobPreset)
+void FTakesView::AddTakesToIngestQueue(const TArray<TObjectPtr<UTakeVirtualAsset>>& InTakeAssets, const UIngestJobSettingsBase& InJobPreset)
 {
 	using namespace UE::CaptureManager;
 
@@ -401,9 +401,8 @@ void FTakesView::AddTakesToIngestQueue(const TArray<TStrongObjectPtr<UTakeVirtua
 	IngestJobs.Reserve(InTakeAssets.Num());
 
 	const FPipelineDetails Pipeline = Controller->GetIngestPipelineManager()->GetSelectedPipeline();
-	for (const TStrongObjectPtr<UTakeVirtualAsset>& TakeAsset : InTakeAssets)
+	for (const TObjectPtr<UTakeVirtualAsset>& TakeAsset : InTakeAssets)
 	{
-
 		EImageRotation ImageRotationToApply =
 			TakeAsset->Metadata.Video.IsEmpty() ? InJobPreset.ImageRotation : ConvertOrientation(TakeAsset->Metadata.Video[0].Orientation);
 
@@ -441,27 +440,46 @@ void FTakesView::AddTakesToIngestQueue(const TArray<TStrongObjectPtr<UTakeVirtua
 	}
 }
 
-void FTakesView::RefreshTakes()
+void FTakesView::Refresh()
 {
 	using namespace UE::CaptureManager;
 
-	if (SelectedItem.IsValid())
+	if (SelectedItem.IsValid() && SelectedItem->Implements<ULiveLinkDeviceCapability_Ingest>())
 	{
-		if (SelectedItem->Implements<ULiveLinkDeviceCapability_Ingest>())
-		{
-			FGuid DeviceId = SelectedItem->GetDeviceId();
+		RefreshSingleTake(SelectedItem.Get());
+	}
+	else
+	{
+		RefreshAllTakes();
+	}
+}
 
-			ResetTakesCache(DeviceId);
+void FTakesView::RefreshSingleTake(TObjectPtr<ULiveLinkDevice> InDevice)
+{
+	ResetTakesCache(InDevice->GetDeviceId());
 
-			FIngestUpdateTakeListCallback Callback =
-				FIngestUpdateTakeListCallback::Type::CreateRaw(this, &FTakesView::UpdateTakeListCallback, DeviceId);
+	FIngestUpdateTakeListCallback Callback =
+		FIngestUpdateTakeListCallback::Type::CreateRaw(this, &FTakesView::UpdateTakeListCallback, InDevice->GetDeviceId());
 
-			UIngestCapability_UpdateTakeListCallback* UpdateTakeListCallback = NewObject<UIngestCapability_UpdateTakeListCallback>();
-			UpdateTakeListCallback->Callback = MoveTemp(Callback);
+	UIngestCapability_UpdateTakeListCallback* UpdateTakeListCallback = NewObject<UIngestCapability_UpdateTakeListCallback>();
+	UpdateTakeListCallback->Callback = MoveTemp(Callback);
 
-			TakesTileView->UpdateTakeListStarted();
-			ILiveLinkDeviceCapability_Ingest::Execute_UpdateTakeList(SelectedItem.Get(), UpdateTakeListCallback);
-		}
+	TakesTileView->UpdateTakeListStarted();
+	ILiveLinkDeviceCapability_Ingest::Execute_UpdateTakeList(InDevice, UpdateTakeListCallback);
+}
+
+void FTakesView::RefreshAllTakes()
+{
+	TSharedPtr<FCaptureManagerPanelController> Controller = PanelController.Pin();
+
+	if (!Controller)
+	{
+		return;
+	}
+
+	for (const TObjectPtr<ULiveLinkDevice>& Device : Controller->GetCaptureDevices())
+	{
+		RefreshSingleTake(Device);
 	}
 }
 
@@ -473,17 +491,17 @@ void FTakesView::UpdateTakeListCallback(TArray<UE::CaptureManager::FTakeId> InTa
 
 	TakesTileView->UpdateTakeListFinished();
 
-	TArray<TStrongObjectPtr<UTakeVirtualAsset>>* Found = CaptureDeviceTakesMap.Find(InCaptureDeviceId);
+	TArray<TUniqueObjectPtr<UTakeVirtualAsset>>* Found = CaptureDeviceTakesMap.Find(InCaptureDeviceId);
 	if (!Found)
 	{
 		return;
 	}
 
-	const TArray<TStrongObjectPtr<UTakeVirtualAsset>>& TakeAssets = *Found;
+	const TArray<TUniqueObjectPtr<UTakeVirtualAsset>>& TakeAssets = *Found;
 
 	if (TakeAssets.IsEmpty() && !InTakeIds.IsEmpty())
 	{
-		TArray<TStrongObjectPtr<UTakeVirtualAsset>> TakeObjects = MakeTakeObjects(InCaptureDeviceId, InTakeIds);
+		TArray<TUniqueObjectPtr<UTakeVirtualAsset>> TakeObjects = MakeTakeObjects(InCaptureDeviceId, InTakeIds);
 
 		CaptureDeviceTakesMap.Add(InCaptureDeviceId, MoveTemp(TakeObjects));
 
@@ -492,7 +510,7 @@ void FTakesView::UpdateTakeListCallback(TArray<UE::CaptureManager::FTakeId> InTa
 	}
 
 	bool bAreEqual = EqualTo(TakeAssets, InTakeIds, 
-							 [](const TStrongObjectPtr<UTakeVirtualAsset>& InLeft, FTakeId InTakeId)
+							 [](const TUniqueObjectPtr<UTakeVirtualAsset>& InLeft, FTakeId InTakeId)
 							 {
 								 return InLeft->TakeId == InTakeId;
 							 });
@@ -569,7 +587,7 @@ void FTakesView::CreateTakesView()
 
 	SAssignNew(TakesTileView, STakesView, MoveTemp(IngestPipelineManager), MoveTemp(IngestJobSettingsManager), MoveTemp(GetCurrentSelectionDelegate))
 		.OnAddTakesToIngestQueue_Raw(this, &FTakesView::AddTakesToIngestQueue)
-		.OnRefreshTakes_Raw(this, &FTakesView::RefreshTakes)
+		.OnRefreshTakes_Raw(this, &FTakesView::Refresh)
 		.TakesPickerConfig(TakesPickerConfig);
 }
 
@@ -674,41 +692,64 @@ void FTakesView::CaptureDeviceAdded(ULiveLinkDevice* InDevice)
 						StaticCastSharedPtr<const FIngestCapability_TakeAddedEvent>(InTakeAddedEvent);
 
 					TArray<FTakeId> TakeIds = { Event->TakeId };
-					TArray<TStrongObjectPtr<UTakeVirtualAsset>> NewTakeObject = MakeTakeObjects(CaptureDeviceId, { TakeIds });
+					TArray<TUniqueObjectPtr<UTakeVirtualAsset>> NewTakeObject = MakeTakeObjects(CaptureDeviceId, { TakeIds });
 
-					if (TArray<TStrongObjectPtr<UTakeVirtualAsset>>* TakeObjects = CaptureDeviceTakesMap.Find(CaptureDeviceId))
+					if (TArray<TUniqueObjectPtr<UTakeVirtualAsset>>* TakeObjects = CaptureDeviceTakesMap.Find(CaptureDeviceId))
 					{
-						TakeObjects->Append(NewTakeObject.GetData(), NewTakeObject.Num());
+						TakeObjects->Append(MoveTemp(NewTakeObject));
 
 						RefreshTakesViewDelegate.ExecuteIfBound(true);
 					}
 				}));
 
 		CapabilityInterface->SubscribeToEvent(FString(FIngestCapability_TakeUpdatedEvent::Name),
-			FCaptureEventHandler([this, CaptureDeviceId](TSharedPtr<const FCaptureEvent> InTakeUpdatedEvent)
+			FCaptureEventHandler([this, CaptureDeviceId, CapabilityInterface](TSharedPtr<const FCaptureEvent> InTakeUpdatedEvent)
 				{
 					check(IsInGameThread());
 
 					TSharedPtr<const FIngestCapability_TakeUpdatedEvent> Event =
 						StaticCastSharedPtr<const FIngestCapability_TakeUpdatedEvent>(InTakeUpdatedEvent);
 
-					TArray<FTakeId> TakeIds = { Event->TakeId };
-					TArray<TStrongObjectPtr<UTakeVirtualAsset>> NewTakeObject = MakeTakeObjects(CaptureDeviceId, { TakeIds });
-					check(NewTakeObject.Num() == 1);
+					auto FindPredicate = 
+						[TakeId = Event->TakeId](const TUniqueObjectPtr<UTakeVirtualAsset>& InTakeObject)
+						{
+							return InTakeObject->TakeId == TakeId;
+						};
 
-					if (TArray<TStrongObjectPtr<UTakeVirtualAsset>>* TakeObjects = CaptureDeviceTakesMap.Find(CaptureDeviceId))
+					TArray<TUniqueObjectPtr<UTakeVirtualAsset>>* MaybeTakeObjects = CaptureDeviceTakesMap.Find(CaptureDeviceId);
+					if (!MaybeTakeObjects)
 					{
-						for (TStrongObjectPtr<UTakeVirtualAsset>& TakeVirtualAsset : *TakeObjects)
-						{
-							if (TakeVirtualAsset->TakeId == Event->TakeId)
-							{
-								TakeVirtualAsset = NewTakeObject[0];
-								break;
-							}
-						}
+						return;
+					}
 
-						RefreshTakesViewDelegate.ExecuteIfBound(true);
+					TUniqueObjectPtr<UTakeVirtualAsset>* MaybeTakeObject = MaybeTakeObjects->FindByPredicate(FindPredicate);
+					if (!MaybeTakeObject)
+					{
+						return;
 					}
+
+					TUniqueObjectPtr<UTakeVirtualAsset>& TakeObject = *MaybeTakeObject;
+
+					TOptional<FTakeMetadata> TakeInfoOpt = CapabilityInterface->GetTakeMetadata(Event->TakeId);
+					if (!TakeInfoOpt.IsSet())
+					{
+						return;
+					}
+
+					FTakeMetadata TakeInfo = TakeInfoOpt.GetValue();
+
+					TakeObject->CaptureDeviceId = CaptureDeviceId;
+					TakeObject->TakeId = Event->TakeId;
+					TakeObject->Metadata = MoveTemp(TakeInfo);
+
+					TOptional<FTakeThumbnail> TakeThumbnail = CreateThumbnail(TakeObject->Metadata.Thumbnail);
+
+					if (TakeThumbnail.IsSet())
+					{
+						TakeObject->Thumbnail = MoveTemp(TakeThumbnail.GetValue());
+					}
+
+					RefreshTakesViewDelegate.ExecuteIfBound(true);
 				}));
 
 		CapabilityInterface->SubscribeToEvent(FString(FIngestCapability_TakeRemovedEvent::Name),
@@ -719,9 +760,9 @@ void FTakesView::CaptureDeviceAdded(ULiveLinkDevice* InDevice)
 					TSharedPtr<const FIngestCapability_TakeRemovedEvent> Event =
 						StaticCastSharedPtr<const FIngestCapability_TakeRemovedEvent>(InTakeRemovedEvent);
 
-					if (TArray<TStrongObjectPtr<UTakeVirtualAsset>>* TakeObjects = CaptureDeviceTakesMap.Find(CaptureDeviceId))
+					if (TArray<TUniqueObjectPtr<UTakeVirtualAsset>>* TakeObjects = CaptureDeviceTakesMap.Find(CaptureDeviceId))
 					{
-						int32 NumRemoved = TakeObjects->RemoveAll([this, TakeId = Event->TakeId](const TStrongObjectPtr<UTakeVirtualAsset>& InObject)
+						int32 NumRemoved = TakeObjects->RemoveAll([this, TakeId = Event->TakeId](const TUniqueObjectPtr<UTakeVirtualAsset>& InObject)
 							{
 								return InObject->TakeId == TakeId;
 							});
@@ -748,9 +789,9 @@ void FTakesView::CaptureDeviceToBeRemoved(ULiveLinkDevice* InDevice)
 		});
 }
 
-TArray<TStrongObjectPtr<UTakeVirtualAsset>> FTakesView::MakeTakeObjects(FGuid InCaptureDeviceId, TOptional<TArray<UE::CaptureManager::FTakeId>> InTakeIds) 
+TArray<TUniqueObjectPtr<UTakeVirtualAsset>> FTakesView::MakeTakeObjects(FGuid InCaptureDeviceId, TOptional<TArray<UE::CaptureManager::FTakeId>> InTakeIds)
 {
-	TArray<TStrongObjectPtr<UTakeVirtualAsset>> TakeObjects;
+	TArray<TUniqueObjectPtr<UTakeVirtualAsset>> TakeObjects;
 
 	TSharedPtr<FCaptureManagerPanelController> Controller = PanelController.Pin();
 	if (!Controller)
@@ -804,7 +845,7 @@ TArray<TStrongObjectPtr<UTakeVirtualAsset>> FTakesView::MakeTakeObjects(FGuid In
 		FTakeMetadata TakeInfo = TakeInfoOpt.GetValue();
 
 		FName Name = MakeUniqueTakeName(TakeInfo, Package);
-		TStrongObjectPtr<UTakeVirtualAsset> TakeObject(NewObject<UTakeVirtualAsset>(Package, Name));
+		TUniqueObjectPtr<UTakeVirtualAsset> TakeObject(NewObject<UTakeVirtualAsset>(Package, Name));
 
 		TakeObject->CaptureDeviceId = InCaptureDeviceId;
 		TakeObject->TakeId = TakeId;
diff --git a/Engine/Plugins/VirtualProduction/CaptureManager/CaptureManagerApp/Source/CaptureManagerEditor/Private/UI/TakesView.h b/Engine/Plugins/VirtualProduction/CaptureManager/CaptureManagerApp/Source/CaptureManagerEditor/Private/UI/TakesView.h
--- a/Engine/Plugins/VirtualProduction/CaptureManager/CaptureManagerApp/Source/CaptureManagerEditor/Private/UI/TakesView.h
+++ b/Engine/Plugins/VirtualProduction/CaptureManager/CaptureManagerApp/Source/CaptureManagerEditor/Private/UI/TakesView.h
@@ -14,6 +14,8 @@
 #include "ContentBrowserDelegates.h"
 #include "IContentBrowserSingleton.h"
 
+#include "TUniqueObjectPtr.h"
+
 #include "IngestManagement/IngestJobSettingsManager.h"
 #include "IngestManagement/IngestPipelineManager.h"
 
@@ -33,7 +35,7 @@ class SIngestJobProcessor;
 class STakesView : public SCompoundWidget
 {
 public:
-	DECLARE_DELEGATE_TwoParams(FOnAddTakesToIngestQueue, const TArray<TStrongObjectPtr<UTakeVirtualAsset>>&, const UIngestJobSettingsBase&);
+	DECLARE_DELEGATE_TwoParams(FOnAddTakesToIngestQueue, const TArray<TObjectPtr<UTakeVirtualAsset>>&, const UIngestJobSettingsBase&);
 	DECLARE_DELEGATE(FOnRefreshTakes);
 
 	SLATE_BEGIN_ARGS(STakesView) {}
@@ -60,7 +62,7 @@ private:
 	FReply OnAddToQueueButtonClicked();
 	TSharedRef<SWidget> OnGeneratePipelineNameWidget(TSharedRef<FText> InText);
 	void OnPipelineSelectionChanged(TSharedPtr<FText> InText, ESelectInfo::Type InSelectType);
-	TArray<TStrongObjectPtr<UTakeVirtualAsset>> GetSelectedTakeAssets() const;
+	TArray<TObjectPtr<UTakeVirtualAsset>> GetSelectedTakeAssets() const;
 
 	TUniquePtr<FGetCurrentSelectionDelegate> GetCurrentSelectionDelegate;
 	TSharedPtr<SPositiveActionButton> RefreshButton;
@@ -91,22 +93,24 @@ private:
 	void CreateTakesView();
 	TWeakPtr<FCaptureManagerPanelController> PanelController;
 
-	TMap<FGuid, TArray<TStrongObjectPtr<UTakeVirtualAsset>>> CaptureDeviceTakesMap;
+	TMap<FGuid, TArray<TUniqueObjectPtr<UTakeVirtualAsset>>> CaptureDeviceTakesMap;
 	TWeakObjectPtr<ULiveLinkDevice> SelectedItem;
 
 	FRefreshAssetViewDelegate RefreshTakesViewDelegate;
 	FSyncToAssetsDelegate SyncToDevicesDelegate;
 
 	TArray<FAssetData> CreateTakeAssetsDataForCaptureDevice(TOptional<FString> InCaptureDeviceName = TOptional<FString>());
-	TArray<FAssetData> CreateTakeAssetsData(const TArray<TStrongObjectPtr<UTakeVirtualAsset>>& InSelectedCaptureDevicesTakesArray);
+	TArray<FAssetData> CreateTakeAssetsData(const TArray<TObjectPtr<UTakeVirtualAsset>>& InSelectedCaptureDevicesTakesArray);
 
 	EImageRotation ConvertOrientation(TOptional<FTakeMetadata::FVideo::EOrientation> InOrientation) const;
-	void AddTakesToIngestQueue(const TArray<TStrongObjectPtr<UTakeVirtualAsset>>& InTakeAssets, const UIngestJobSettingsBase& InJobPreset);
-	void RefreshTakes();
+	void AddTakesToIngestQueue(const TArray<TObjectPtr<UTakeVirtualAsset>>& InTakeAssets, const UIngestJobSettingsBase& InJobPreset);
+	void Refresh();
+	void RefreshSingleTake(TObjectPtr<ULiveLinkDevice> InDevice);
+	void RefreshAllTakes();
 	void UpdateTakeListCallback(TArray<UE::CaptureManager::FTakeId> InTakeIds, FGuid InCaptureDeviceId);
 	void ResetTakesCache(FGuid InCaptureDeviceId);
 
-	TArray<TStrongObjectPtr<UTakeVirtualAsset>> MakeTakeObjects(FGuid InCaptureDeviceId,
+	TArray<TUniqueObjectPtr<UTakeVirtualAsset>> MakeTakeObjects(FGuid InCaptureDeviceId,
 																TOptional<TArray<UE::CaptureManager::FTakeId>> InTakeIds = TOptional<TArray<UE::CaptureManager::FTakeId>>());
 	TOptional<FTakeThumbnail> CreateThumbnail(const FTakeThumbnailData& InThumbnailData);
 
