[CaptureManager] Fix crash when ingesting legacy archive footage

diff --git a/Engine/Plugins/VirtualProduction/CaptureManager/CaptureManagerApp/Source/CaptureDataConverter/Private/Nodes/ThirdPartyEncoder/CaptureConvertVideoDataThirdParty.cpp b/Engine/Plugins/VirtualProduction/CaptureManager/CaptureManagerApp/Source/CaptureDataConverter/Private/Nodes/ThirdPartyEncoder/CaptureConvertVideoDataThirdParty.cpp
--- a/Engine/Plugins/VirtualProduction/CaptureManager/CaptureManagerApp/Source/CaptureDataConverter/Private/Nodes/ThirdPartyEncoder/CaptureConvertVideoDataThirdParty.cpp
+++ b/Engine/Plugins/VirtualProduction/CaptureManager/CaptureManagerApp/Source/CaptureDataConverter/Private/Nodes/ThirdPartyEncoder/CaptureConvertVideoDataThirdParty.cpp
@@ -301,7 +301,7 @@ FCaptureConvertVideoDataThirdParty::FResult FCaptureConvertVideoDataThirdParty::
 	}
 
 	FTaskProgress::FTask Task = Params.TaskProgress->StartTask();
-	FTaskProgress TaskProgress(Video.FramesCount.GetValue(), FTaskProgress::FProgressReporter::CreateLambda([&Task](float InProgress)
+	TSharedPtr<FTaskProgress> TaskProgress = MakeShared<FTaskProgress>(Video.FramesCount.GetValue(), FTaskProgress::FProgressReporter::CreateLambda([&Task](float InProgress)
 	{
 		Task.Update(InProgress);
 	}));
@@ -326,7 +326,7 @@ FCaptureConvertVideoDataThirdParty::FResult FCaptureConvertVideoDataThirdParty::
 			return false;
 		}
 		
-		FTaskProgress::FTask Task = TaskProgress.StartTask();
+		FTaskProgress::FTask Task = TaskProgress->StartTask();
 
 		FCopyProgressReporter ProgressReporter(Task, Params.StopToken);
 
