Devices not being added properly after config is loaded

diff --git a/Engine/Plugins/VirtualProduction/CaptureManager/CaptureManagerApp/Source/CaptureManagerEditor/Private/UI/CaptureManagerPanelController.cpp b/Engine/Plugins/VirtualProduction/CaptureManager/CaptureManagerApp/Source/CaptureManagerEditor/Private/UI/CaptureManagerPanelController.cpp
--- a/Engine/Plugins/VirtualProduction/CaptureManager/CaptureManagerApp/Source/CaptureManagerEditor/Private/UI/CaptureManagerPanelController.cpp
+++ b/Engine/Plugins/VirtualProduction/CaptureManager/CaptureManagerApp/Source/CaptureManagerEditor/Private/UI/CaptureManagerPanelController.cpp
@@ -200,6 +200,27 @@ void FCaptureManagerPanelController::OnDeviceAdded(FGuid InDeviceId, ULiveLinkDe
 	}
 
 	TakesView->CaptureDeviceAdded(InDevice);
+
+	if (SourcesReachableMap.Contains(InDeviceId) && SourcesReachableMap[InDeviceId])
+	{
+		TakesView->CaptureDeviceStarted(InDeviceId);
+	}
+	else if (!SourcesReachableMap.Contains(InDeviceId))
+	{
+		ELiveLinkDeviceConnectionStatus ConnectionStatus = ILiveLinkDeviceCapability_Connection::Execute_GetConnectionStatus(InDevice);
+
+		bool bConnectionStatus = ConnectionStatus == ELiveLinkDeviceConnectionStatus::Connected ? true : false;
+		SourcesReachableMap.Add(InDeviceId, bConnectionStatus);
+
+		if (ConnectionStatus == ELiveLinkDeviceConnectionStatus::Connected)
+		{
+			TakesView->CaptureDeviceStarted(InDeviceId);
+		}
+		else
+		{
+			ILiveLinkDeviceCapability_Connection::Execute_Connect(InDevice);
+		}
+	}
 }
 
 void FCaptureManagerPanelController::OnDeviceRemoved(FGuid InDeviceId, ULiveLinkDevice* InDevice)
@@ -216,6 +237,7 @@ void FCaptureManagerPanelController::OnDeviceRemoved(FGuid InDeviceId, ULiveLink
 	}
 
 	TakesView->CaptureDeviceToBeRemoved(InDevice);
+	SourcesReachableMap.Remove(InDeviceId);
 }
 
 void FCaptureManagerPanelController::OnReachableEvent(ELiveLinkDeviceConnectionStatus InStatus, FGuid InDeviceId)
diff --git a/Engine/Plugins/VirtualProduction/CaptureManager/CaptureManagerDevices/Source/CPSLiveLinkDevice/Private/LiveLinkFaceDevice.cpp b/Engine/Plugins/VirtualProduction/CaptureManager/CaptureManagerDevices/Source/CPSLiveLinkDevice/Private/LiveLinkFaceDevice.cpp
--- a/Engine/Plugins/VirtualProduction/CaptureManager/CaptureManagerDevices/Source/CPSLiveLinkDevice/Private/LiveLinkFaceDevice.cpp
+++ b/Engine/Plugins/VirtualProduction/CaptureManager/CaptureManagerDevices/Source/CPSLiveLinkDevice/Private/LiveLinkFaceDevice.cpp
@@ -74,6 +74,11 @@ void ULiveLinkFaceDevice::UpdateTakeList_Implementation(UIngestCapability_Update
 
 	RemoveAllTakes();
 
+	if (!Device)
+	{
+		return;
+	}
+
 	TProtocolResult<TArray<FGetTakeMetadataResponse::FTakeObject>> TakesResult = Device->FetchTakeList();
 
 	if (TakesResult.HasError())
