Initiate connect before the device is configured

diff --git a/Engine/Plugins/VirtualProduction/CaptureManager/CaptureManagerApp/Source/CaptureManagerEditor/Private/UI/CaptureManagerPanelController.cpp b/Engine/Plugins/VirtualProduction/CaptureManager/CaptureManagerApp/Source/CaptureManagerEditor/Private/UI/CaptureManagerPanelController.cpp
--- a/Engine/Plugins/VirtualProduction/CaptureManager/CaptureManagerApp/Source/CaptureManagerEditor/Private/UI/CaptureManagerPanelController.cpp
+++ b/Engine/Plugins/VirtualProduction/CaptureManager/CaptureManagerApp/Source/CaptureManagerEditor/Private/UI/CaptureManagerPanelController.cpp
@@ -209,7 +209,7 @@ void FCaptureManagerPanelController::OnDeviceAdded(FGuid InDeviceId, ULiveLinkDe
 	{
 		ELiveLinkDeviceConnectionStatus ConnectionStatus = ILiveLinkDeviceCapability_Connection::Execute_GetConnectionStatus(InDevice);
 
-		bool bConnectionStatus = ConnectionStatus == ELiveLinkDeviceConnectionStatus::Connected ? true : false;
+		bool bConnectionStatus = ConnectionStatus == ELiveLinkDeviceConnectionStatus::Connected;
 		SourcesReachableMap.Add(InDeviceId, bConnectionStatus);
 
 		if (ConnectionStatus == ELiveLinkDeviceConnectionStatus::Connected)
@@ -244,7 +244,7 @@ void FCaptureManagerPanelController::OnReachableEvent(ELiveLinkDeviceConnectionS
 {
 	using namespace UE::CaptureManager;
 
-	bool ConnectionStatus = InStatus == ELiveLinkDeviceConnectionStatus::Connected ? true : false;
+	bool ConnectionStatus = InStatus == ELiveLinkDeviceConnectionStatus::Connected;
 	SourcesReachableMap.Add(InDeviceId, ConnectionStatus);
 
 	if (ConnectionStatus)
diff --git a/Engine/Plugins/VirtualProduction/CaptureManager/CaptureManagerDevices/Source/CPSLiveLinkDevice/Private/LiveLinkFaceDevice.cpp b/Engine/Plugins/VirtualProduction/CaptureManager/CaptureManagerDevices/Source/CPSLiveLinkDevice/Private/LiveLinkFaceDevice.cpp
--- a/Engine/Plugins/VirtualProduction/CaptureManager/CaptureManagerDevices/Source/CPSLiveLinkDevice/Private/LiveLinkFaceDevice.cpp
+++ b/Engine/Plugins/VirtualProduction/CaptureManager/CaptureManagerDevices/Source/CPSLiveLinkDevice/Private/LiveLinkFaceDevice.cpp
@@ -200,6 +200,11 @@ bool ULiveLinkFaceDevice::Connect_Implementation()
 	const ULiveLinkFaceDeviceSettings* DeviceSettings = GetSettings();
 
 	using namespace UE::CaptureManager;
+	
+	if (DeviceSettings->IpAddress.IpAddressString.IsEmpty())
+	{
+		return false;
+	}
 
 	Device = FCPSDevice::MakeCPSDevice(DeviceSettings->IpAddress.IpAddressString, static_cast<uint16>(DeviceSettings->Port));
 	
