[CaptureManager] Add support to extract timecode using user enabled third-party encoder.

diff --git a/Engine/Plugins/VirtualProduction/CaptureManager/CaptureManagerApp/Source/IngestLiveLinkDevice/Private/Utils/CaptureExtractTimecode.cpp b/Engine/Plugins/VirtualProduction/CaptureManager/CaptureManagerApp/Source/IngestLiveLinkDevice/Private/Utils/CaptureExtractTimecode.cpp
--- a/Engine/Plugins/VirtualProduction/CaptureManager/CaptureManagerApp/Source/IngestLiveLinkDevice/Private/Utils/CaptureExtractTimecode.cpp
+++ b/Engine/Plugins/VirtualProduction/CaptureManager/CaptureManagerApp/Source/IngestLiveLinkDevice/Private/Utils/CaptureExtractTimecode.cpp
@@ -8,54 +8,58 @@
 #include "MediaPlaylist.h"
 #include "Sound/SoundWaveTimecodeInfo.h"
 #include "UObject/Package.h"
+#include "Settings/CaptureManagerSettings.h"
+#include "Logging/LogMacros.h"
 
-FCaptureExtractTimecode::FCaptureExtractTimecode(const FString& InFilePath)
-	: FilePath(InFilePath)
+#define LOCTEXT_NAMESPACE "CaptureExtractTimecode"
+
+DECLARE_LOG_CATEGORY_EXTERN(LogCaptureExtractTimecode, Log, All);
+DEFINE_LOG_CATEGORY(LogCaptureExtractTimecode);
+
+FCaptureExtractVideoTimecode::FCaptureExtractVideoTimecode(const FString& InFilePath)
 {}
 
-FCaptureExtractTimecode::FTimecodeAndRateResult FCaptureExtractTimecode::Extract()
+UE::CaptureManager::FTimecodeAndRateResult FCaptureExtractVideoTimecode::Extract()
 {	
+	using namespace UE::CaptureManager;
+
 	check(!FilePath.IsEmpty())
 
 	FTimecodeAndRateResult Result = MakeError(ECaptureExtractTimecodeError::UnhandledMedia);
 
-	FString FileExtension = FPaths::GetExtension(FilePath);
-	if (FileExtension == "wav")
-	{
-		Result = ExtractTimecodeFromWavBFWSection();
-	}
-	else
-	{
-		Result = ExtractTimecodeUsingElectraPlayer();
-	}
+	Result = ExtractTimecodeUsingElectraPlayer();
 
-	return Result;
-}
-
-FCaptureExtractTimecode::FTimecodeAndRateResult FCaptureExtractTimecode::ExtractTimecodeFromWavBFWSection()
-{
-	TArray<uint8> WAVData;
-	if (FFileHelper::LoadFileToArray(WAVData, *FilePath))
+	if (!Result.IsValid())
 	{
-		FWaveModInfo WAVInfo;
-		if (WAVInfo.ReadWaveInfo(WAVData.GetData(), WAVData.Num()))
+		const UCaptureManagerSettings* Settings = GetDefault<UCaptureManagerSettings>();
+
+		if (Settings->bEnableThirdPartyEncoder)
 		{
-			if (WAVInfo.TimecodeInfo.IsValid())
+			FString EncoderPath = Settings->ThirdPartyEncoder.FilePath;
+			if (EncoderPath.EndsWith(TEXT("ffmpeg.exe")))
 			{
-				FSoundWaveTimecodeInfo TimecodeInfo = *WAVInfo.TimecodeInfo.Get();
-				const double NumSecondsSinceMidnight = TimecodeInfo.GetNumSecondsSinceMidnight();
-				FTimecode AudioTimecode = FTimecode(NumSecondsSinceMidnight, TimecodeInfo.TimecodeRate, TimecodeInfo.bTimecodeIsDropFrame, /* InbRollover = */ true);
-				FTimecodeAndRate TimecodeAndRate{ AudioTimecode, TimecodeInfo.TimecodeRate };
-				return MakeValue(TimecodeAndRate);
+				EncoderPath.ReplaceInline(TEXT("ffmpeg.exe"), TEXT("ffprobe.exe"));
+				Result = ExtractTimecodeUsingFFProbe(EncoderPath);
 			}
 		}
 	}
 
-	return MakeError(ECaptureExtractTimecodeError::TimecodeNotFound);
+	if (Result.HasValue())
+	{
+		UE_LOG(LogCaptureExtractTimecode, Display, TEXT("Extracted timecode '%s' from video file: %s"), *Result.GetValue().Timecode.ToString(), *FilePath);
+	}
+	else
+	{
+		UE_LOG(LogCaptureExtractTimecode, Display, TEXT("Timecode not found for video file: %s"), *FilePath);
+	}
+
+	return Result;
 }
 
-FCaptureExtractTimecode::FTimecodeAndRateResult FCaptureExtractTimecode::ExtractTimecodeUsingElectraPlayer()
+UE::CaptureManager::FTimecodeAndRateResult FCaptureExtractVideoTimecode::ExtractTimecodeUsingElectraPlayer()
 {
+	using namespace UE::CaptureManager;
+
 	//Some media players require waiting a frame before metadata can be read
 	bGetTimecodeDelay = true;
 
@@ -64,7 +68,6 @@ FCaptureExtractTimecode::FTimecodeAndRateResult FCaptureExtractTimecode::Extract
 
 	AsyncTask(ENamedThreads::GameThread, [this, &Promise]
 		{
-
 			UDesiredPlayerMediaSource* MediaSource = NewObject<UDesiredPlayerMediaSource>();
 			MediaSource->SetFilePath(FilePath);
 
@@ -88,7 +91,7 @@ FCaptureExtractTimecode::FTimecodeAndRateResult FCaptureExtractTimecode::Extract
 			}
 
 			Context = { MediaPlayer, Promise };
-			
+
 			// Sometimes it is necessary to wait a couple of ticks for the MediaPlayer to be ready.
 			// Log start time so that we can timeout if initialization is taking too long.
 			MediaPlayerWaitStart = FDateTime::Now();
@@ -99,9 +102,11 @@ FCaptureExtractTimecode::FTimecodeAndRateResult FCaptureExtractTimecode::Extract
 	return Future.Get();
 }
 
-void FCaptureExtractTimecode::TickOutput(FTimespan DeltaTime, FTimespan Timecode)
+void FCaptureExtractVideoTimecode::TickOutput(FTimespan InDeltaTime, FTimespan InTimespan)
 {
 #if WITH_EDITOR
+	using namespace UE::CaptureManager;
+
 	if (Context.IsSet())
 	{
 		TStrongObjectPtr<UMediaPlayer>& MediaPlayer = Context.GetValue().MediaPlayer;
@@ -155,13 +160,13 @@ void FCaptureExtractTimecode::TickOutput(FTimespan DeltaTime, FTimespan Timecode
 		FTimecodeAndRate TimecodeAndRate;
 		if (TimecodeResult.IsValid())
 		{
-			TimecodeAndRate.MaybeTimecode = TimecodeResult.StealValue();
-		
+			TimecodeAndRate.Timecode = TimecodeResult.StealValue();
+
 			if (TimecodeRateResult.IsValid())
 			{
 				TimecodeAndRate.MaybeTimecodeRate = TimecodeRateResult.StealValue();
 			}
-			
+
 			Result = MakeValue(MoveTemp(TimecodeAndRate));
 		}
 
@@ -170,7 +175,7 @@ void FCaptureExtractTimecode::TickOutput(FTimespan DeltaTime, FTimespan Timecode
 #endif // WITH_EDITOR
 }
 
-void FCaptureExtractTimecode::InitializeClockSink()
+void FCaptureExtractVideoTimecode::InitializeClockSink()
 {
 	IMediaModule* MediaModule = FModuleManager::LoadModulePtr<IMediaModule>("Media");
 	if (MediaModule != nullptr)
@@ -179,7 +184,7 @@ void FCaptureExtractTimecode::InitializeClockSink()
 	}
 }
 
-void FCaptureExtractTimecode::DeinitializeClockSink()
+void FCaptureExtractVideoTimecode::DeinitializeClockSink()
 {
 	IMediaModule* MediaModule = FModuleManager::GetModulePtr<IMediaModule>("Media");
 	if (MediaModule != nullptr)
@@ -188,8 +193,10 @@ void FCaptureExtractTimecode::DeinitializeClockSink()
 	}
 }
 
-TValueOrError<FTimecode, FCaptureExtractTimecode::ECaptureExtractTimecodeError> FCaptureExtractTimecode::GetTimecode(TStrongObjectPtr<UMediaPlayer>& MediaPlayer)
+TValueOrError<FTimecode, UE::CaptureManager::ECaptureExtractTimecodeError> FCaptureExtractVideoTimecode::GetTimecode(TStrongObjectPtr<UMediaPlayer>& MediaPlayer)
 {
+	using namespace UE::CaptureManager;
+
 	// Timecode from the source. Either player's meta data or media source.
 	const FVariant Value = MediaPlayer->GetMediaInfo(UMediaPlayer::MediaInfoNameStartTimecodeValue.Resolve());
 	FString TimecodeString;
@@ -200,7 +207,6 @@ TValueOrError<FTimecode, FCaptureExtractTimecode::ECaptureExtractTimecodeError>
 	}
 	else
 	{
-
 		if (UMediaPlaylist* Playlist = MediaPlayer->GetPlaylist())
 		{
 			if (const UMediaSource* MediaSource = Playlist->Get(MediaPlayer->GetPlaylistIndex()))
@@ -226,8 +232,10 @@ TValueOrError<FTimecode, FCaptureExtractTimecode::ECaptureExtractTimecodeError>
 	}
 }
 
-TValueOrError<FFrameRate, FCaptureExtractTimecode::ECaptureExtractTimecodeError> FCaptureExtractTimecode::GetTimecodeRate(TStrongObjectPtr<UMediaPlayer>& MediaPlayer)
+TValueOrError<FFrameRate, UE::CaptureManager::ECaptureExtractTimecodeError> FCaptureExtractVideoTimecode::GetTimecodeRate(TStrongObjectPtr<UMediaPlayer>& MediaPlayer)
 {
+	using namespace UE::CaptureManager;
+
 	// Timecode from the source. Either player's meta data or media source.
 	const FVariant Value = MediaPlayer->GetMediaInfo(UMediaPlayer::MediaInfoNameStartTimecodeFrameRate.Resolve());
 	FString TimecodeRateString;
@@ -297,3 +305,278 @@ TValueOrError<FFrameRate, FCaptureExtractTimecode::ECaptureExtractTimecodeError>
 		return MakeValue(TimecodeRate);
 	}
 }
+
+UE::CaptureManager::FTimecodeAndRateResult FCaptureExtractVideoTimecode::ExtractTimecodeUsingFFProbe(const FString& InFFProbePath)
+{
+	using namespace UE::CaptureManager;
+
+	FTimecodeAndRateResult Result = MakeError(ECaptureExtractTimecodeError::UnhandledMedia);
+
+	FString CommandArgs = TEXT("-v error -select_streams v:0 -show_entries stream_tags=timecode:stream=r_frame_rate -of default=noprint_wrappers=1:nokey=1 ") + FilePath;
+
+	void* ReadPipe = nullptr;
+	void* WritePipe = nullptr;
+	verify(FPlatformProcess::CreatePipe(ReadPipe, WritePipe, false));
+
+	constexpr bool bLaunchDetached = false;
+	constexpr bool bLaunchHidden = true;
+	constexpr bool bLaunchReallyHidden = true;
+	FProcHandle ProcHandle =
+		FPlatformProcess::CreateProc(*InFFProbePath, *CommandArgs, bLaunchDetached, bLaunchHidden, bLaunchReallyHidden, nullptr, 0, nullptr, WritePipe, nullptr);
+
+	ON_SCOPE_EXIT
+	{
+		FPlatformProcess::TerminateProc(ProcHandle);
+		FPlatformProcess::CloseProc(ProcHandle);
+		FPlatformProcess::ClosePipe(ReadPipe, WritePipe);
+	};
+
+	if (!ProcHandle.IsValid())
+	{
+		UE_LOG(LogCaptureExtractTimecode, Display, TEXT("Extract timecode: Failed to start the process %s %s"), *InFFProbePath, *CommandArgs);
+		return MakeError(ECaptureExtractTimecodeError::InternalError);
+	}
+
+	const FDateTime WaitStart = FDateTime::Now();
+
+	TArray<uint8> FullCommandOutput;
+	while (FPlatformProcess::IsProcRunning(ProcHandle))
+	{
+		TArray<uint8> CommandOutput;
+		{
+			bool Read = FPlatformProcess::ReadPipeToArray(ReadPipe, CommandOutput);
+			if (!Read)
+			{
+				CommandOutput.Empty();
+			}
+		}
+
+		if ((FDateTime::Now() - WaitStart).GetSeconds() > TimeoutPeriod)
+		{
+			// Timed out
+			break;
+		}
+
+		if (CommandOutput.IsEmpty())
+		{
+			FPlatformProcess::Sleep(0.1f);
+			continue;
+		}
+
+		FullCommandOutput.Append(MoveTemp(CommandOutput));
+
+	}
+
+	int32 ReturnCode = 0;
+	FPlatformProcess::GetProcReturnCode(ProcHandle, &ReturnCode);
+
+
+	TArray<uint8> CommandOutput;
+	{
+		bool Read = FPlatformProcess::ReadPipeToArray(ReadPipe, CommandOutput);
+		if (!Read)
+		{
+			CommandOutput.Empty();
+		}
+	}
+	FullCommandOutput.Append(CommandOutput);
+
+	if (ReturnCode == 0)
+	{
+		FTimecode Timecode;
+		bool bFoundTimecode = false;
+		FFrameRate FrameRate;
+		bool bFoundFrameRate = false;
+
+		FString CommandOutputStr = FString::ConstructFromPtrSize(reinterpret_cast<const UTF8CHAR*>(FullCommandOutput.GetData()), FullCommandOutput.Num());
+
+		// Split the timecode and frame rate lines from the output
+		TArray<FString> OutputLinesStr;
+		CommandOutputStr.ParseIntoArrayLines(OutputLinesStr);
+
+		for (const FString& Line : OutputLinesStr)
+		{
+			if (!bFoundFrameRate)
+			{
+				TArray<FString> FrameRatePartsStr;
+				Line.ParseIntoArray(FrameRatePartsStr, TEXT("/"));
+				if (FrameRatePartsStr.Num() == 2)
+				{
+					int32 Numerator = FCString::Atoi(*FrameRatePartsStr[0]);
+					int32 Denominator = FCString::Atoi(*FrameRatePartsStr[1]);
+
+					FrameRate = FFrameRate(Numerator, Denominator);
+					bFoundFrameRate = true;
+					continue;
+				}
+			}
+			
+			if (!bFoundTimecode)
+			{
+				TArray<FString> TimecodePartsStr;
+				Line.ParseIntoArray(TimecodePartsStr, TEXT(":")); // We don't support DF timecode
+				if (TimecodePartsStr.Num() == 4)
+				{
+					int32 Hours = FCString::Atoi(*TimecodePartsStr[0]);
+					int32 Minutes = FCString::Atoi(*TimecodePartsStr[1]);
+					int32 Seconds = FCString::Atoi(*TimecodePartsStr[2]);
+					int32 Frames = FCString::Atoi(*TimecodePartsStr[3]);
+
+					Timecode = FTimecode(Hours, Minutes, Seconds, Frames, true);
+
+					bFoundTimecode = true;
+					continue;
+				}
+			}
+		}
+
+		if (!bFoundTimecode)
+		{
+			return MakeError(ECaptureExtractTimecodeError::UnableToParseTimecode);
+		}
+
+		if (!bFoundFrameRate)
+		{
+			return MakeError(ECaptureExtractTimecodeError::UnableToParseTimecodeRate);
+		}
+
+		Result = MakeValue(FTimecodeAndRate{ Timecode , FrameRate });
+	}
+	else
+	{
+		if (!FullCommandOutput.IsEmpty())
+		{
+			UE_LOG(LogCaptureExtractTimecode, Display, TEXT("Failed to run the command: %s %s"), *InFFProbePath, *CommandArgs);
+
+			FString CommandOutputStr = FString::ConstructFromPtrSize(reinterpret_cast<const UTF8CHAR*>(FullCommandOutput.GetData()), FullCommandOutput.Num());
+			UE_LOG(LogCaptureExtractTimecode, Display,
+				TEXT("Output from the command:\n>>>>>>\n%s<<<<<<"), *CommandOutputStr);
+		}
+
+		return MakeError(ECaptureExtractTimecodeError::InternalError);
+	}
+
+	return Result;
+}
+
+FCaptureExtractAudioTimecode::FCaptureExtractAudioTimecode(const FString& InFilePath)
+{}
+
+UE::CaptureManager::FTimecodeAndRateResult FCaptureExtractAudioTimecode::Extract()
+{
+	return ExtractTimecodeFromBroadcastWaveFormat(FFrameRate());
+}
+
+static FFrameRate EstimateSmpteTimecodeRate(const FFrameRate InMediaFrameRate)
+{
+	if (FMath::IsNearlyEqual(InMediaFrameRate.AsDecimal(), 60.0))
+	{
+		return FFrameRate(30'000, 1'000);
+	}
+
+	if (FMath::IsNearlyEqual(InMediaFrameRate.AsDecimal(), 59.94))
+	{
+		// 29.97
+		return FFrameRate(30'000, 1'001);
+	}
+
+	if (FMath::IsNearlyEqual(InMediaFrameRate.AsDecimal(), 50.0))
+	{
+		FFrameRate(25'000, 1'000);
+	}
+
+	if (FMath::IsNearlyEqual(InMediaFrameRate.AsDecimal(), 48.0))
+	{
+		FFrameRate(24'000, 1'000);
+	}
+
+	return InMediaFrameRate;
+}
+
+UE::CaptureManager::FTimecodeAndRateResult FCaptureExtractAudioTimecode::Extract(FFrameRate InFrameRate)
+{
+	using namespace UE::CaptureManager;
+
+	check(!FilePath.IsEmpty())
+
+	FTimecodeAndRateResult Result = MakeError(ECaptureExtractTimecodeError::UnhandledMedia);
+
+	FString FileExtension = FPaths::GetExtension(FilePath);
+	if (FileExtension == "wav")
+	{
+		// Convert timecode rate to SMPTE timecode rate
+		FFrameRate TimecodeRate = EstimateSmpteTimecodeRate(InFrameRate);
+		Result = ExtractTimecodeFromBroadcastWaveFormat(TimecodeRate);
+	}
+
+	if (Result.HasValue())
+	{
+		UE_LOG(LogCaptureExtractTimecode, Display, TEXT("Extracted timecode '%s' from audio file: %s"), *Result.GetValue().Timecode.ToString(), *FilePath);
+	}
+	else
+	{
+		UE_LOG(LogCaptureExtractTimecode, Display, TEXT("Timecode not found for audio file: %s"), *FilePath);
+	}
+
+	return Result;
+}
+
+UE::CaptureManager::FTimecodeAndRateResult FCaptureExtractAudioTimecode::ExtractTimecodeFromBroadcastWaveFormat(FFrameRate InTimecodeRate)
+{
+	using namespace UE::CaptureManager;
+
+	TArray<uint8> WAVData;
+	if (FFileHelper::LoadFileToArray(WAVData, *FilePath))
+	{
+		FWaveModInfo WAVInfo;
+		if (WAVInfo.ReadWaveInfo(WAVData.GetData(), WAVData.Num()))
+		{
+			if (WAVInfo.TimecodeInfo.IsValid())
+			{
+				FSoundWaveTimecodeInfo TimecodeInfo = *WAVInfo.TimecodeInfo.Get();
+				const double NumSecondsSinceMidnight = TimecodeInfo.GetNumSecondsSinceMidnight();
+
+				FFrameRate TimecodeRate = TimecodeInfo.TimecodeRate;
+				const bool bTimecodeRateIsSampleRate = TimecodeRate == FFrameRate(TimecodeInfo.NumSamplesPerSecond, 1);
+				if (bTimecodeRateIsSampleRate)
+				{
+					UE_LOG(
+						LogCaptureExtractTimecode,
+						Display,
+						TEXT(
+							"Embedded timecode rate is %.2f fps (the sample rate). "
+							"This usually indicates there is no timecode rate information in the wav file: %s"
+						),
+						TimecodeInfo.TimecodeRate.AsDecimal(),
+						*FilePath
+					);
+
+					if (InTimecodeRate != FFrameRate())
+					{
+						// Use the provided timecode rate instead
+						TimecodeRate = InTimecodeRate;
+
+						UE_LOG(
+							LogCaptureExtractTimecode,
+							Display,
+							TEXT(
+								"Taking the embedded audio timecode but estimating an SMPTE audio timecode rate. "
+								"Timecode rate for %s set to %.2f"
+							),
+							*FilePath,
+							TimecodeRate.AsDecimal()
+						);
+					}
+				}
+
+				FTimecode AudioTimecode = FTimecode(NumSecondsSinceMidnight, TimecodeRate, TimecodeInfo.bTimecodeIsDropFrame, /* InbRollover = */ true);
+				FTimecodeAndRate TimecodeAndRate{ AudioTimecode, TimecodeRate };
+				return MakeValue(TimecodeAndRate);
+			}
+		}
+	}
+
+	return MakeError(ECaptureExtractTimecodeError::TimecodeNotFound);
+}
+
+#undef LOCTEXT_NAMESPACE
diff --git a/Engine/Plugins/VirtualProduction/CaptureManager/CaptureManagerApp/Source/IngestLiveLinkDevice/Public/Utils/CaptureExtractTimecode.h b/Engine/Plugins/VirtualProduction/CaptureManager/CaptureManagerApp/Source/IngestLiveLinkDevice/Public/Utils/CaptureExtractTimecode.h
--- a/Engine/Plugins/VirtualProduction/CaptureManager/CaptureManagerApp/Source/IngestLiveLinkDevice/Public/Utils/CaptureExtractTimecode.h
+++ b/Engine/Plugins/VirtualProduction/CaptureManager/CaptureManagerApp/Source/IngestLiveLinkDevice/Public/Utils/CaptureExtractTimecode.h
@@ -17,6 +17,27 @@
 
 class UMediaPlayer;
 
+namespace UE::CaptureManager
+{
+struct FTimecodeAndRate
+{
+	FTimecode Timecode;
+	TOptional<FFrameRate> MaybeTimecodeRate;
+};
+
+enum ECaptureExtractTimecodeError : int32
+{
+	InternalError = -1,
+	TimecodeNotFound = 1,
+	UnableToParseTimecode = 2,
+	UnableToParseTimecodeRate = 3,
+	TimecodeRateNotFound = 4,
+	UnhandledMedia = 5,
+};
+
+using FTimecodeAndRateResult = TValueOrError<FTimecodeAndRate, ECaptureExtractTimecodeError>;
+}
+
 #if WITH_EDITOR
 
 UCLASS(BlueprintType, MinimalAPI)
@@ -36,37 +57,17 @@ public:
 };
 #endif // WITH_EDITOR
 
-class INGESTLIVELINKDEVICE_API FCaptureExtractTimecode final : public TSharedFromThis<FCaptureExtractTimecode>, public IMediaClockSink
+class INGESTLIVELINKDEVICE_API FCaptureExtractVideoTimecode final : public TSharedFromThis<FCaptureExtractVideoTimecode>, public IMediaClockSink
 {
 public:
 
-	struct FTimecodeAndRate
-	{
-		TOptional<FTimecode> MaybeTimecode;
-		TOptional<FFrameRate> MaybeTimecodeRate;
-	};
-
-	enum ECaptureExtractTimecodeError : int32
-	{
-		InternalError = -1,
-		TimecodeNotFound = 1,
-		UnableToParseTimecode = 2,
-		UnableToParseTimecodeRate = 3,
-		TimecodeRateNotFound = 4,
-		UnhandledMedia = 5,
-	};
+	FCaptureExtractVideoTimecode(const FString& InFilePath);
 
-	using FTimecodeAndRateResult = TValueOrError<FTimecodeAndRate, ECaptureExtractTimecodeError>;
-
-	FCaptureExtractTimecode(const FString& InFilePath);
-
-	FTimecodeAndRateResult Extract();
+	UE::CaptureManager::FTimecodeAndRateResult Extract();
 
 private:
 
-	FString FilePath;
-
-	using FTimecodeAndRatePromise = TPromise<FTimecodeAndRateResult>;
+	using FTimecodeAndRatePromise = TPromise<UE::CaptureManager::FTimecodeAndRateResult>;
 
 	struct FTimecodeExtractionContext
 	{
@@ -74,18 +75,20 @@ private:
 		FTimecodeAndRatePromise& ExtractedTimecodeAndRatePromise;
 	};
 
+	FString FilePath;
+
 	TOptional<FTimecodeExtractionContext> Context;
 
 #if WITH_EDITOR
-	TValueOrError<FTimecode, ECaptureExtractTimecodeError> GetTimecode(TStrongObjectPtr<UMediaPlayer>& MediaPlayer);
-	TValueOrError<FFrameRate, ECaptureExtractTimecodeError> GetTimecodeRate(TStrongObjectPtr<UMediaPlayer>& MediaPlayer);
+	TValueOrError<FTimecode, UE::CaptureManager::ECaptureExtractTimecodeError> GetTimecode(TStrongObjectPtr<UMediaPlayer>& MediaPlayer);
+	TValueOrError<FFrameRate, UE::CaptureManager::ECaptureExtractTimecodeError> GetTimecodeRate(TStrongObjectPtr<UMediaPlayer>& MediaPlayer);
 #endif
 	void InitializeClockSink();
 	void DeinitializeClockSink();
 
-	FTimecodeAndRateResult ExtractTimecodeFromWavBFWSection();
-	FTimecodeAndRateResult ExtractTimecodeUsingElectraPlayer();
-	
+	UE::CaptureManager::FTimecodeAndRateResult ExtractTimecodeUsingElectraPlayer();
+	UE::CaptureManager::FTimecodeAndRateResult ExtractTimecodeUsingFFProbe(const FString& InEncoderPath);
+
 	bool bGetTimecodeDelay;
 	const int32 TimeoutPeriod = 3;
 	FDateTime MediaPlayerWaitStart;
@@ -95,3 +98,22 @@ private:
 	 */
 	virtual void TickOutput(FTimespan DeltaTime, FTimespan Timecode) override;
 };
+
+class INGESTLIVELINKDEVICE_API FCaptureExtractAudioTimecode final : public TSharedFromThis<FCaptureExtractAudioTimecode>
+{
+public:
+
+	FCaptureExtractAudioTimecode(const FString& InFilePath);
+	~FCaptureExtractAudioTimecode() = default;
+
+	UE::CaptureManager::FTimecodeAndRateResult Extract();
+	UE::CaptureManager::FTimecodeAndRateResult Extract(FFrameRate InFrameRate);
+
+private:
+
+	FString FilePath;
+
+	const int32 TimeoutPeriod = 3;
+
+	UE::CaptureManager::FTimecodeAndRateResult ExtractTimecodeFromBroadcastWaveFormat(FFrameRate InTimecodeRate);
+};
diff --git a/Engine/Plugins/VirtualProduction/CaptureManager/CaptureManagerDevices/Source/CPSLiveLinkDevice/CPSLiveLinkDevice.Build.cs b/Engine/Plugins/VirtualProduction/CaptureManager/CaptureManagerDevices/Source/CPSLiveLinkDevice/CPSLiveLinkDevice.Build.cs
--- a/Engine/Plugins/VirtualProduction/CaptureManager/CaptureManagerDevices/Source/CPSLiveLinkDevice/CPSLiveLinkDevice.Build.cs
+++ b/Engine/Plugins/VirtualProduction/CaptureManager/CaptureManagerDevices/Source/CPSLiveLinkDevice/CPSLiveLinkDevice.Build.cs
@@ -22,6 +22,7 @@ public class CPSLiveLinkDevice : ModuleRules
 			"Json",
 			"CaptureManagerMediaRW",
 			"CaptureProtocolStack",
+			"DataIngestCore",
 			"IngestLiveLinkDevice",
 			"LiveLinkDevice",
 			"LiveLinkCapabilities",
diff --git a/Engine/Plugins/VirtualProduction/CaptureManager/CaptureManagerDevices/Source/CPSLiveLinkDevice/Private/LiveLinkFaceDevice.cpp b/Engine/Plugins/VirtualProduction/CaptureManager/CaptureManagerDevices/Source/CPSLiveLinkDevice/Private/LiveLinkFaceDevice.cpp
--- a/Engine/Plugins/VirtualProduction/CaptureManager/CaptureManagerDevices/Source/CPSLiveLinkDevice/Private/LiveLinkFaceDevice.cpp
+++ b/Engine/Plugins/VirtualProduction/CaptureManager/CaptureManagerDevices/Source/CPSLiveLinkDevice/Private/LiveLinkFaceDevice.cpp
@@ -15,6 +15,7 @@
 #include "LiveLinkFaceMetadata.h"
 #include "StereoCameraMetadataParseUtils.h"
 #include "Utils/CaptureExtractTimecode.h"
+#include "Utils/ParseTakeUtils.h"
 
 #include "HAL/FileManager.h"
 
@@ -563,53 +564,42 @@ void ULiveLinkFaceDevice::RemoveDownloadedTakeData(const UE::CaptureManager::FTa
 
 void ULiveLinkFaceDevice::ExtractTimecodeIfNotSet(FTakeMetadata& InOutTakeMetadata)
 {
-	for (auto& Video : InOutTakeMetadata.Video)
+	using namespace UE::CaptureManager;
+
+	bool bVideoFrameRateSet = false;
+	FFrameRate VideoFrameRate;
+	for (FTakeMetadata::FVideo& Video : InOutTakeMetadata.Video)
 	{
 		if (!Video.TimecodeStart.IsSet())
 		{
-			TSharedPtr<FCaptureExtractTimecode> Extractor = MakeShareable(new FCaptureExtractTimecode(Video.Path));
-			FCaptureExtractTimecode::FTimecodeAndRateResult TimecodeAndRateResult = Extractor->Extract();
+			TSharedPtr<FCaptureExtractVideoTimecode> Extractor = MakeShareable(new FCaptureExtractVideoTimecode(Video.Path));
+			FTimecodeAndRateResult TimecodeAndRateResult = Extractor->Extract();
 			if (TimecodeAndRateResult.IsValid())
 			{
-				FCaptureExtractTimecode::FTimecodeAndRate TimecodeAndRate = TimecodeAndRateResult.GetValue();
-				if (TimecodeAndRate.MaybeTimecode.IsSet())
-				{
-					Video.TimecodeStart = TimecodeAndRate.MaybeTimecode.GetValue().ToString();
-				}
+				FTimecodeAndRate TimecodeAndRate = TimecodeAndRateResult.GetValue();
+				Video.TimecodeStart = TimecodeAndRate.Timecode.ToString();
 			}
 		}
-	}
 
-	for (auto& Depth : InOutTakeMetadata.Depth)
-	{
-		if (!Depth.TimecodeStart.IsSet())
+		if (!bVideoFrameRateSet)
 		{
-			TSharedPtr<FCaptureExtractTimecode> Extractor = MakeShareable(new FCaptureExtractTimecode(Depth.Path));
-			FCaptureExtractTimecode::FTimecodeAndRateResult TimecodeAndRateResult = Extractor->Extract();
-			if (TimecodeAndRateResult.IsValid())
-			{
-				FCaptureExtractTimecode::FTimecodeAndRate TimecodeAndRate = TimecodeAndRateResult.GetValue();
-				if (TimecodeAndRate.MaybeTimecode.IsSet())
-				{
-					Depth.TimecodeStart = TimecodeAndRate.MaybeTimecode.GetValue().ToString();
-				}
-			}
+			VideoFrameRate = UE::CaptureManager::ParseFrameRate(Video.FrameRate);
 		}
 	}
 
-	for (auto& Audio : InOutTakeMetadata.Audio)
+	for (FTakeMetadata::FAudio& Audio : InOutTakeMetadata.Audio)
 	{
 		if (!Audio.TimecodeStart.IsSet() && !Audio.TimecodeRate.IsSet())
 		{
-			TSharedPtr<FCaptureExtractTimecode> Extractor = MakeShareable(new FCaptureExtractTimecode(Audio.Path));
-			FCaptureExtractTimecode::FTimecodeAndRateResult TimecodeAndRateResult = Extractor->Extract();
+			TSharedPtr<FCaptureExtractAudioTimecode> Extractor = MakeShareable(new FCaptureExtractAudioTimecode(Audio.Path));
+
+			// The video frame rate will be used to calculate the timecode rate if the timecode rate cannot be extracted from the audio file
+			FTimecodeAndRateResult TimecodeAndRateResult = Extractor->Extract(VideoFrameRate); 
 			if (TimecodeAndRateResult.IsValid())
 			{
-				FCaptureExtractTimecode::FTimecodeAndRate TimecodeAndRate = TimecodeAndRateResult.GetValue();
-				if (TimecodeAndRate.MaybeTimecode.IsSet())
-				{
-					Audio.TimecodeStart = TimecodeAndRate.MaybeTimecode.GetValue().ToString();
-				}
+				FTimecodeAndRate TimecodeAndRate = TimecodeAndRateResult.GetValue();
+				Audio.TimecodeStart = TimecodeAndRate.Timecode.ToString();
+
 				if (TimecodeAndRate.MaybeTimecodeRate.IsSet())
 				{
 					Audio.TimecodeRate = static_cast<float>(TimecodeAndRate.MaybeTimecodeRate.GetValue().AsDecimal());
diff --git a/Engine/Plugins/VirtualProduction/CaptureManager/CaptureManagerDevices/Source/FilesystemLiveLinkDevice/FilesystemLiveLinkDevice.Build.cs b/Engine/Plugins/VirtualProduction/CaptureManager/CaptureManagerDevices/Source/FilesystemLiveLinkDevice/FilesystemLiveLinkDevice.Build.cs
--- a/Engine/Plugins/VirtualProduction/CaptureManager/CaptureManagerDevices/Source/FilesystemLiveLinkDevice/FilesystemLiveLinkDevice.Build.cs
+++ b/Engine/Plugins/VirtualProduction/CaptureManager/CaptureManagerDevices/Source/FilesystemLiveLinkDevice/FilesystemLiveLinkDevice.Build.cs
@@ -22,6 +22,7 @@ public class FilesystemLiveLinkDevice : ModuleRules
 			"Json",
 			"CaptureManagerMediaRW",
 			"IngestLiveLinkDevice",
+			"DataIngestCore",
 			"LiveLinkDevice",
 			"LiveLinkCapabilities",
 			"Media",
diff --git a/Engine/Plugins/VirtualProduction/CaptureManager/CaptureManagerDevices/Source/FilesystemLiveLinkDevice/Private/FilesystemLiveLinkDevice.cpp b/Engine/Plugins/VirtualProduction/CaptureManager/CaptureManagerDevices/Source/FilesystemLiveLinkDevice/Private/FilesystemLiveLinkDevice.cpp
--- a/Engine/Plugins/VirtualProduction/CaptureManager/CaptureManagerDevices/Source/FilesystemLiveLinkDevice/Private/FilesystemLiveLinkDevice.cpp
+++ b/Engine/Plugins/VirtualProduction/CaptureManager/CaptureManagerDevices/Source/FilesystemLiveLinkDevice/Private/FilesystemLiveLinkDevice.cpp
@@ -12,6 +12,7 @@
 
 #include "Utils/IngestLiveLinkDeviceUtils.h"
 #include "Utils/CaptureExtractTimecode.h"
+#include "Utils/ParseTakeUtils.h"
 
 DEFINE_LOG_CATEGORY(LogGenericLiveLinkArchiveDevice);
 
@@ -241,6 +242,8 @@ TOptional<FTakeMetadata> UFilesystemLiveLinkDevice::ReadTake(const FString& InCu
 		}
 	}
 	
+	bool bVideoFrameRateSet = false;
+	FFrameRate VideoFrameRate;
 	for (FTakeMetadata::FVideo& Video : TakeMetadata.Video)
 	{
 		if (FPaths::IsRelative(*Video.Path))
@@ -248,16 +251,21 @@ TOptional<FTakeMetadata> UFilesystemLiveLinkDevice::ReadTake(const FString& InCu
 			Video.Path = FPaths::ConvertRelativePathToFull(CurrentDirectory, *Video.Path);
 		}
 
-		TSharedPtr<FCaptureExtractTimecode> Extractor = MakeShareable(new FCaptureExtractTimecode(Video.Path));
-		FCaptureExtractTimecode::FTimecodeAndRateResult TimecodeAndRateResult = Extractor->Extract();
+		TSharedPtr<FCaptureExtractVideoTimecode> Extractor = MakeShareable(new FCaptureExtractVideoTimecode(Video.Path));
+		FTimecodeAndRateResult TimecodeAndRateResult = Extractor->Extract();
 		if (TimecodeAndRateResult.IsValid())
 		{
-			FCaptureExtractTimecode::FTimecodeAndRate TimecodeAndRate = TimecodeAndRateResult.GetValue();
-			if (!Video.TimecodeStart.IsSet() && TimecodeAndRate.MaybeTimecode.IsSet())
+			FTimecodeAndRate TimecodeAndRate = TimecodeAndRateResult.GetValue();
+			if (!Video.TimecodeStart.IsSet())
 			{
-				Video.TimecodeStart = TimecodeAndRate.MaybeTimecode.GetValue().ToString();
+				Video.TimecodeStart = TimecodeAndRate.Timecode.ToString();
 			}
 		}
+
+		if (!bVideoFrameRateSet)
+		{
+			VideoFrameRate = UE::CaptureManager::ParseFrameRate(Video.FrameRate);
+		}
 	}
 	for (FTakeMetadata::FVideo& Depth : TakeMetadata.Depth)
 	{
@@ -266,14 +274,14 @@ TOptional<FTakeMetadata> UFilesystemLiveLinkDevice::ReadTake(const FString& InCu
 			Depth.Path = FPaths::ConvertRelativePathToFull(CurrentDirectory, *Depth.Path);
 		}
 
-		TSharedPtr<FCaptureExtractTimecode> Extractor = MakeShareable(new FCaptureExtractTimecode(Depth.Path));
-		FCaptureExtractTimecode::FTimecodeAndRateResult TimecodeAndRateResult = Extractor->Extract();
+		TSharedPtr<FCaptureExtractVideoTimecode> Extractor = MakeShareable(new FCaptureExtractVideoTimecode(Depth.Path));
+		FTimecodeAndRateResult TimecodeAndRateResult = Extractor->Extract();
 		if (TimecodeAndRateResult.IsValid())
 		{
-			FCaptureExtractTimecode::FTimecodeAndRate TimecodeAndRate = TimecodeAndRateResult.GetValue();
-			if (!Depth.TimecodeStart.IsSet() && TimecodeAndRate.MaybeTimecode.IsSet())
+			FTimecodeAndRate TimecodeAndRate = TimecodeAndRateResult.GetValue();
+			if (!Depth.TimecodeStart.IsSet())
 			{
-				Depth.TimecodeStart = TimecodeAndRate.MaybeTimecode.GetValue().ToString();
+				Depth.TimecodeStart = TimecodeAndRate.Timecode.ToString();
 			}
 		}
 	}
@@ -285,15 +293,16 @@ TOptional<FTakeMetadata> UFilesystemLiveLinkDevice::ReadTake(const FString& InCu
 			Audio.Path = FPaths::ConvertRelativePathToFull(CurrentDirectory, *Audio.Path);
 		}
 
-		TSharedPtr<FCaptureExtractTimecode> Extractor = MakeShareable(new FCaptureExtractTimecode(Audio.Path));
-		FCaptureExtractTimecode::FTimecodeAndRateResult TimecodeAndRateResult = Extractor->Extract();
+		TSharedPtr<FCaptureExtractAudioTimecode> Extractor = MakeShareable(new FCaptureExtractAudioTimecode(Audio.Path));
+
+		// The video frame rate will be used to calculate the timecode rate if the timecode rate cannot be extracted from the audio file
+		// If the VideoFrameRate is not set, the timecode rate may be set to the audio sample rate
+		FTimecodeAndRateResult TimecodeAndRateResult = Extractor->Extract(VideoFrameRate);
 		if (TimecodeAndRateResult.IsValid())
 		{
-			FCaptureExtractTimecode::FTimecodeAndRate TimecodeAndRate = TimecodeAndRateResult.GetValue();
-			if (TimecodeAndRate.MaybeTimecode.IsSet())
-			{
-				Audio.TimecodeStart = TimecodeAndRate.MaybeTimecode.GetValue().ToString();
-			}
+			FTimecodeAndRate TimecodeAndRate = TimecodeAndRateResult.GetValue();
+			Audio.TimecodeStart = TimecodeAndRate.Timecode.ToString();
+			
 			if (TimecodeAndRate.MaybeTimecodeRate.IsSet())
 			{
 				Audio.TimecodeRate = static_cast<float>(TimecodeAndRate.MaybeTimecodeRate.GetValue().AsDecimal());
diff --git a/Engine/Plugins/VirtualProduction/CaptureManager/CaptureManagerDevices/Source/MonoVideoLiveLinkDevice/Private/MonoVideoLiveLinkDevice.cpp b/Engine/Plugins/VirtualProduction/CaptureManager/CaptureManagerDevices/Source/MonoVideoLiveLinkDevice/Private/MonoVideoLiveLinkDevice.cpp
--- a/Engine/Plugins/VirtualProduction/CaptureManager/CaptureManagerDevices/Source/MonoVideoLiveLinkDevice/Private/MonoVideoLiveLinkDevice.cpp
+++ b/Engine/Plugins/VirtualProduction/CaptureManager/CaptureManagerDevices/Source/MonoVideoLiveLinkDevice/Private/MonoVideoLiveLinkDevice.cpp
@@ -311,15 +311,12 @@ TOptional<FTakeMetadata> UMonoVideoLiveLinkDevice::ReadTake(const FString& InCur
 	Video.Path = InCurrentTakeFile;
 	Video.PathType = FTakeMetadata::FVideo::EPathType::File;
 
-	TSharedPtr<FCaptureExtractTimecode> Extractor = MakeShareable(new FCaptureExtractTimecode(Video.Path));
-	FCaptureExtractTimecode::FTimecodeAndRateResult TimecodeAndRateResult = Extractor->Extract();
+	TSharedPtr<FCaptureExtractVideoTimecode> Extractor = MakeShareable(new FCaptureExtractVideoTimecode(Video.Path));
+	FTimecodeAndRateResult TimecodeAndRateResult = Extractor->Extract();
 	if (TimecodeAndRateResult.IsValid())
 	{
-		FCaptureExtractTimecode::FTimecodeAndRate TimecodeAndRate = TimecodeAndRateResult.GetValue();
-		if (TimecodeAndRate.MaybeTimecode.IsSet())
-		{
-			Video.TimecodeStart = TimecodeAndRate.MaybeTimecode.GetValue().ToString();
-		}
+		FTimecodeAndRate TimecodeAndRate = TimecodeAndRateResult.GetValue();
+		Video.TimecodeStart = TimecodeAndRate.Timecode.ToString();
 	}
 
 	TakeMetadata.Video.Add(MoveTemp(Video));
diff --git a/Engine/Plugins/VirtualProduction/CaptureManager/CaptureManagerDevices/Source/StereoVideoLiveLinkDevice/Private/StereoVideoLiveLinkDevice.cpp b/Engine/Plugins/VirtualProduction/CaptureManager/CaptureManagerDevices/Source/StereoVideoLiveLinkDevice/Private/StereoVideoLiveLinkDevice.cpp
--- a/Engine/Plugins/VirtualProduction/CaptureManager/CaptureManagerDevices/Source/StereoVideoLiveLinkDevice/Private/StereoVideoLiveLinkDevice.cpp
+++ b/Engine/Plugins/VirtualProduction/CaptureManager/CaptureManagerDevices/Source/StereoVideoLiveLinkDevice/Private/StereoVideoLiveLinkDevice.cpp
@@ -14,6 +14,7 @@
 #include "Utils/TakeDiscoveryExpressionParser.h"
 #include "Utils/VideoDeviceThumbnailExtractor.h"
 #include "Utils/CaptureExtractTimecode.h"
+#include "Utils/ParseTakeUtils.h"
 
 #include "VideoLiveLinkDeviceLog.h"
 
@@ -472,6 +473,9 @@ TOptional<FTakeMetadata> UStereoVideoLiveLinkDevice::CreateTakeMetadata(const FT
 
 	int32 AudioNameCounter = 1;
 
+	bool bVideoFrameRateSet = false;
+	FFrameRate VideoFrameRate;
+
 	for (const FTakeWithComponents::Component& TakeComponent : InTakeWithComponents.Components)
 	{
 		if (TakeComponent.Type == ETakeComponentType::VIDEO || TakeComponent.Type == ETakeComponentType::IMAGE_SEQUENCE)
@@ -497,6 +501,12 @@ TOptional<FTakeMetadata> UStereoVideoLiveLinkDevice::CreateTakeMetadata(const FT
 				{
 					TUniquePtr<IVideoReader> VideoReader = VideoReaderResult.StealValue();
 					Video.FrameRate = static_cast<float>(VideoReader->GetFrameRate().AsDecimal());
+
+					if (!bVideoFrameRateSet)
+					{
+						VideoFrameRate = UE::CaptureManager::ParseFrameRate(Video.FrameRate);
+					}
+
 					TValueOrError<TUniquePtr<UE::CaptureManager::FMediaTextureSample>, FText> TextureSampleResult = VideoReader->Next();
 					if (TextureSampleResult.HasValue())
 					{
@@ -525,15 +535,12 @@ TOptional<FTakeMetadata> UStereoVideoLiveLinkDevice::CreateTakeMetadata(const FT
 				Video.Path = TakeComponent.Path;
 				
 				FString MediaFilePath = FPaths::ConvertRelativePathToFull(InTakeWithComponents.TakePath, Video.Path);
-				TSharedPtr<FCaptureExtractTimecode> Extractor = MakeShareable(new FCaptureExtractTimecode(MediaFilePath));
-				FCaptureExtractTimecode::FTimecodeAndRateResult TimecodeAndRateResult = Extractor->Extract();
+				TSharedPtr<FCaptureExtractVideoTimecode> Extractor = MakeShareable(new FCaptureExtractVideoTimecode(MediaFilePath));
+				FTimecodeAndRateResult TimecodeAndRateResult = Extractor->Extract();
 				if (TimecodeAndRateResult.IsValid())
 				{
-					FCaptureExtractTimecode::FTimecodeAndRate TimecodeAndRate = TimecodeAndRateResult.GetValue();
-					if (TimecodeAndRate.MaybeTimecode.IsSet())
-					{
-						Video.TimecodeStart = TimecodeAndRate.MaybeTimecode.GetValue().ToString();
-					}
+					FTimecodeAndRate TimecodeAndRate = TimecodeAndRateResult.GetValue();
+					Video.TimecodeStart = TimecodeAndRate.Timecode.ToString();
 				}
 			}
 			else
@@ -597,15 +604,15 @@ TOptional<FTakeMetadata> UStereoVideoLiveLinkDevice::CreateTakeMetadata(const FT
 			Audio.Duration = 0;
 
 			FString MediaFilePath = FPaths::ConvertRelativePathToFull(InTakeWithComponents.TakePath, Audio.Path);
-			TSharedPtr<FCaptureExtractTimecode> Extractor = MakeShareable(new FCaptureExtractTimecode(MediaFilePath));
-			FCaptureExtractTimecode::FTimecodeAndRateResult TimecodeAndRateResult = Extractor->Extract();
+			TSharedPtr<FCaptureExtractAudioTimecode> Extractor = MakeShareable(new FCaptureExtractAudioTimecode(MediaFilePath));
+
+			// The video frame rate will be used to calculate the timecode rate if the timecode rate cannot be extracted from the audio file
+			FTimecodeAndRateResult TimecodeAndRateResult = Extractor->Extract(VideoFrameRate);
 			if (TimecodeAndRateResult.IsValid())
 			{
-				FCaptureExtractTimecode::FTimecodeAndRate TimecodeAndRate = TimecodeAndRateResult.GetValue();
-				if (TimecodeAndRate.MaybeTimecode.IsSet())
-				{
-					Audio.TimecodeStart = TimecodeAndRate.MaybeTimecode.GetValue().ToString();
-				}
+				FTimecodeAndRate TimecodeAndRate = TimecodeAndRateResult.GetValue();
+				Audio.TimecodeStart = TimecodeAndRate.Timecode.ToString();
+			
 				if (TimecodeAndRate.MaybeTimecodeRate.IsSet())
 				{
 					Audio.TimecodeRate = static_cast<float>(TimecodeAndRate.MaybeTimecodeRate.GetValue().AsDecimal());
diff --git a/Engine/Plugins/VirtualProduction/CaptureManager/CaptureManagerDevices/Source/StereoVideoLiveLinkDevice/StereoVideoLiveLinkDevice.Build.cs b/Engine/Plugins/VirtualProduction/CaptureManager/CaptureManagerDevices/Source/StereoVideoLiveLinkDevice/StereoVideoLiveLinkDevice.Build.cs
--- a/Engine/Plugins/VirtualProduction/CaptureManager/CaptureManagerDevices/Source/StereoVideoLiveLinkDevice/StereoVideoLiveLinkDevice.Build.cs
+++ b/Engine/Plugins/VirtualProduction/CaptureManager/CaptureManagerDevices/Source/StereoVideoLiveLinkDevice/StereoVideoLiveLinkDevice.Build.cs
@@ -22,6 +22,7 @@ public class StereoVideoLiveLinkDevice : ModuleRules
 			"Json",
 			"CaptureManagerMediaRW",
 			"IngestLiveLinkDevice",
+			"DataIngestCore",
 			"LiveLinkDevice",
 			"LiveLinkCapabilities",
 			"Media",
