Memory usage remains high when ingest job is cancelled

diff --git a/Engine/Plugins/VirtualProduction/CaptureManager/CaptureManagerApp/Source/CaptureDataConverter/Private/Nodes/CaptureConvertDepthData.cpp b/Engine/Plugins/VirtualProduction/CaptureManager/CaptureManagerApp/Source/CaptureDataConverter/Private/Nodes/CaptureConvertDepthData.cpp
--- a/Engine/Plugins/VirtualProduction/CaptureManager/CaptureManagerApp/Source/CaptureDataConverter/Private/Nodes/CaptureConvertDepthData.cpp
+++ b/Engine/Plugins/VirtualProduction/CaptureManager/CaptureManagerApp/Source/CaptureDataConverter/Private/Nodes/CaptureConvertDepthData.cpp
@@ -22,14 +22,10 @@ FCaptureConvertDepthData::FCaptureConvertDepthData(const FTakeMetadata::FVideo&
 	: FConvertDepthNode(InDepth, InOutputDirectory)
 	, Params(InParams)
 	, DepthParams(InDepthParams)
-	, WritingQueue(UE::CaptureManager::TQueueRunner<FWritingContext>::FOnProcess::CreateRaw(this, &FCaptureConvertDepthData::OnWrite))
 {
 }
 
-FCaptureConvertDepthData::~FCaptureConvertDepthData()
-{
-	WritingQueue.Empty();
-}
+FCaptureConvertDepthData::~FCaptureConvertDepthData() = default;
 
 FCaptureConvertDepthData::FResult FCaptureConvertDepthData::Run()
 {
@@ -101,41 +97,43 @@ FCaptureConvertDepthData::FResult FCaptureConvertDepthData::ConvertData()
 		ImageWriter->Close();
 	};
 
-	TPromise<FResult> Promise;
-	TFuture<FResult> Future = Promise.GetFuture();
-	bIsDone = false;
-	bool bShouldRun = true;
-
-	while (bShouldRun)
+	while (true)
 	{
 		TValueOrError<TUniquePtr<UE::CaptureManager::FMediaTextureSample>, FText> DepthSampleResult = DepthReader->Next();
 
-		// Stream ended
-		if (DepthSampleResult.HasValue() && !DepthSampleResult.GetValue().IsValid())
+		if (DepthSampleResult.HasError())
 		{
-			// Break when we set the nullptr value to the writer so it is aware of the ending
-			bShouldRun = false;
+			return MakeError(DepthSampleResult.StealError());
+		}
+
+		TUniquePtr<UE::CaptureManager::FMediaTextureSample> DepthSample = DepthSampleResult.StealValue();
+
+		// End of stream
+		if (!DepthSample.IsValid())
+		{
+			break;
 		}
 
 		FWritingContext Context;
-		Context.ReadResult = MoveTemp(DepthSampleResult);
+		Context.ReadSample = MoveTemp(DepthSample);
 		Context.Writer = ImageWriter.Get();
-		
+
 		if (Depth.FramesCount.IsSet())
 		{
 			Context.TotalDuration = Depth.FramesCount.GetValue();
 		}
 
 		Context.Task = &Task;
-		Context.DonePromise = &Promise;
 
-		WritingQueue.Add(MoveTemp(Context));
-	}
+		FWriteResult WriteResult = OnWrite(MoveTemp(Context));
 
-	FResult Result = Future.Get();
-	check(bIsDone); // Check that writing thread finished
+		if (WriteResult.HasError())
+		{
+			return MakeError(WriteResult.StealError());
+		}
+	}
 
-	return Result;
+	return MakeValue();
 }
 
 FCaptureConvertDepthData::FResult FCaptureConvertDepthData::CopyData()
@@ -212,47 +210,16 @@ FCaptureConvertDepthData::FResult FCaptureConvertDepthData::CopyData()
 	return MakeValue();
 }
 
-void FCaptureConvertDepthData::OnWrite(FWritingContext InContext)
+FCaptureConvertDepthData::FWriteResult FCaptureConvertDepthData::OnWrite(FWritingContext InContext)
 {
-	if (bIsDone)
-	{
-		return;
-	}
-
-	TPromise<FResult>* DonePromise = InContext.DonePromise;
-
-	if (InContext.ReadResult.HasError())
-	{
-		bIsDone = true;
-
-		DonePromise->SetValue(MakeError(InContext.ReadResult.StealError()));
-		return;
-	}
-
-	TUniquePtr<UE::CaptureManager::FMediaTextureSample> DepthSample = InContext.ReadResult.StealValue();
-
-	// End of stream
-	if (!DepthSample.IsValid())
-	{
-		bIsDone = true;
-
-		DonePromise->SetValue(MakeValue());
-		return;
-	}
-
-	IImageWriter* Writer = InContext.Writer;
-
+	TUniquePtr<UE::CaptureManager::FMediaTextureSample>& DepthSample = InContext.ReadSample;
 	DepthSample->Rotation = DepthParams.Rotation;
 
-	TOptional<FText> WriterResult = Writer->Append(DepthSample.Get());
+	TOptional<FText> WriterResult = InContext.Writer->Append(DepthSample.Get());
 
 	if (WriterResult.IsSet())
 	{
-		bIsDone = true;
-
-		DonePromise->SetValue(MakeError(MoveTemp(WriterResult.GetValue())));
-
-		return;
+		return MakeError(MoveTemp(WriterResult.GetValue()));
 	}
 
 	++CurrentFrame;
@@ -265,11 +232,11 @@ void FCaptureConvertDepthData::OnWrite(FWritingContext InContext)
 
 	if (Params.StopToken.IsStopRequested())
 	{
-		bIsDone = true;
-
 		FText Message = LOCTEXT("LiveLinkConvertDepthNode_AbortedByUser", "Aborted by user");
-		DonePromise->SetValue(MakeError(MoveTemp(Message)));
+		return MakeError(MoveTemp(Message));
 	}
+
+	return MakeValue();
 }
 
 #undef LOCTEXT_NAMESPACE
\ No newline at end of file
diff --git a/Engine/Plugins/VirtualProduction/CaptureManager/CaptureManagerApp/Source/CaptureDataConverter/Private/Nodes/CaptureConvertDepthData.h b/Engine/Plugins/VirtualProduction/CaptureManager/CaptureManagerApp/Source/CaptureDataConverter/Private/Nodes/CaptureConvertDepthData.h
--- a/Engine/Plugins/VirtualProduction/CaptureManager/CaptureManagerApp/Source/CaptureDataConverter/Private/Nodes/CaptureConvertDepthData.h
+++ b/Engine/Plugins/VirtualProduction/CaptureManager/CaptureManagerApp/Source/CaptureDataConverter/Private/Nodes/CaptureConvertDepthData.h
@@ -8,8 +8,6 @@
 
 #include "MediaSample.h"
 
-#include "Async/QueueRunner.h"
-
 class FCaptureConvertDepthData final
 	: public FConvertDepthNode
 {
@@ -35,26 +33,21 @@ private:
 	struct FWritingContext
 	{
 		FWritingContext()
-			: ReadResult(MakeValue())
+			: ReadSample(nullptr)
 			, Writer(nullptr)
 			, Task(nullptr)
 			, TotalDuration(0)
-			, DonePromise(nullptr)
 		{
 		}
 
-		TValueOrError<TUniquePtr<UE::CaptureManager::FMediaTextureSample>, FText> ReadResult;
+		TUniquePtr<UE::CaptureManager::FMediaTextureSample> ReadSample;
 		class IImageWriter* Writer;
 		UE::CaptureManager::FTaskProgress::FTask* Task;
-
 		TOptional<int32> TotalDuration;
-
-		TPromise<FResult>* DonePromise;
 	};
 
-	void OnWrite(FWritingContext InContext);
+	using FWriteResult = TValueOrError<void, FText>;
+	FWriteResult OnWrite(FWritingContext InContext);
 
-	UE::CaptureManager::TQueueRunner<FWritingContext> WritingQueue;
-	std::atomic_bool bIsDone = false;
 	int32 CurrentFrame = 0;
 };
\ No newline at end of file
diff --git a/Engine/Plugins/VirtualProduction/CaptureManager/CaptureManagerApp/Source/CaptureManagerMediaRW/Private/Writers/DepthImageWriter.cpp b/Engine/Plugins/VirtualProduction/CaptureManager/CaptureManagerApp/Source/CaptureManagerMediaRW/Private/Writers/DepthImageWriter.cpp
--- a/Engine/Plugins/VirtualProduction/CaptureManager/CaptureManagerApp/Source/CaptureManagerMediaRW/Private/Writers/DepthImageWriter.cpp
+++ b/Engine/Plugins/VirtualProduction/CaptureManager/CaptureManagerApp/Source/CaptureManagerMediaRW/Private/Writers/DepthImageWriter.cpp
@@ -111,15 +111,7 @@ TArray<float> FDepthExrImageWriter::Transform(TArray<uint8> InDepthArray, FIntPo
 
 	FIntPoint OutputSize = GetOutputSize(InDimensions, InOrientation);
 
-	TArray<int16> ConvertedData;
-
-	for (int32 Index = 0; Index < InDepthArray.Num(); Index += 2)
-	{
-		int16 Value = *reinterpret_cast<int16*>(InDepthArray.GetData() + Index);
-		ConvertedData.Add(Value);
-	}
-
-	InDepthArray.Empty();
+	TArrayView<int16> ConvertedData(reinterpret_cast<int16*>(InDepthArray.GetData()), InDepthArray.Num() / 2);
 
 	switch (InOrientation)
 	{
diff --git a/Engine/Plugins/VirtualProduction/CaptureManager/CaptureManagerCore/Source/CaptureUtils/Public/Async/QueueRunner.h b/Engine/Plugins/VirtualProduction/CaptureManager/CaptureManagerCore/Source/CaptureUtils/Public/Async/QueueRunner.h
--- a/Engine/Plugins/VirtualProduction/CaptureManager/CaptureManagerCore/Source/CaptureUtils/Public/Async/QueueRunner.h
+++ b/Engine/Plugins/VirtualProduction/CaptureManager/CaptureManagerCore/Source/CaptureUtils/Public/Async/QueueRunner.h
@@ -45,7 +45,7 @@ public:
 	void Empty()
 	{
 		FScopeLock lock(&Mutex);
-		while (Queue.Pop());
+		Queue.Empty();
 	}
 
 	bool IsRunning() const
