Resolving issue with ffmpeg being used when running Python scripts

diff --git a/Engine/Plugins/Developer/NamingTokens/Source/NamingTokens/Private/NamingTokensEngineSubsystem.cpp b/Engine/Plugins/Developer/NamingTokens/Source/NamingTokens/Private/NamingTokensEngineSubsystem.cpp
--- a/Engine/Plugins/Developer/NamingTokens/Source/NamingTokens/Private/NamingTokensEngineSubsystem.cpp
+++ b/Engine/Plugins/Developer/NamingTokens/Source/NamingTokens/Private/NamingTokensEngineSubsystem.cpp
@@ -51,62 +51,20 @@ TArray<FString> UNamingTokensEngineSubsystem::GetAllNamespaces() const
 
 UNamingTokens* UNamingTokensEngineSubsystem::GetNamingTokens(const FString& InNamespace) const
 {
-	auto LoadNamingToken = [this, &InNamespace](const TSoftClassPtr<UNamingTokens>& InTokensClass) -> UNamingTokens*
-	{
-		if (const UClass* Class = InTokensClass.LoadSynchronous())
-		{
-			UNamingTokens* NamingTokens = NewObject<UNamingTokens>(GetTransientPackage(), Class);
-			// Initial validation on load. This will log any errors.
-			NamingTokens->Validate();
-			if (ensure(NamingTokens->GetNamespace() == InNamespace))
-			{
-				CachedNamingTokens.Add(InNamespace, NamingTokens);
-				return NamingTokens;
-			}
-		}
-		return nullptr;
-	};
-
 	// Check cache.
-	if (const TObjectPtr<UNamingTokens>* Tokens = CachedNamingTokens.Find(InNamespace))
+	UNamingTokens* FoundTokens = GetNamingTokenFromCache(InNamespace, false);
+	
+	if (FoundTokens)
 	{
-		return *Tokens;
+		return FoundTokens;
 	}
 
-	UNamingTokens* FoundTokens = nullptr;
+	FoundTokens = GetNamingTokensNative(InNamespace);
 
 	const UClass* TargetClass = UNamingTokens::StaticClass();
-	
-	// Check native classes.
-	{
-		constexpr EClassFlags InvalidClassFlags = CLASS_Deprecated | CLASS_NewerVersionExists | CLASS_Abstract;
-
-		TArray<UClass*> DerivedClasses;
-		GetDerivedClasses(TargetClass, DerivedClasses);
-		for (const UClass* DerivedClass : DerivedClasses)
-		{
-			if (DerivedClass->IsNative()
-				&& DerivedClass->IsChildOf(TargetClass)
-				&& !DerivedClass->HasAnyClassFlags(InvalidClassFlags)
-				&& !DerivedClass->GetName().StartsWith(TEXT("SKEL_"))
-				&& !DerivedClass->GetName().StartsWith(TEXT("REINST_"))
-				&& DerivedClass->GetDefaultObject<UNamingTokens>()->GetNamespace() == InNamespace)
-			{
-				if (!FoundTokens)
-				{
-					FoundTokens = LoadNamingToken(DerivedClass);
-				}
-				else
-				{
-					UE_LOG(LogNamingTokens, Warning, TEXT("Namespace '%s' exists more than once in native class '%s' and '%s'."),
-						*InNamespace, *FoundTokens->GetClass()->GetName(), *DerivedClass->GetName());
-				}
-			}
-		}
-	}
 
 	// Check blueprint classes.
-	// NOTE: We may want to use OnAssetAdded and cache the asset or load the class then.  
+	// NOTE: We may want to use OnAssetAdded and cache the asset or load the class then.
 	{
 		const FAssetRegistryModule& AssetRegistryModule = FModuleManager::LoadModuleChecked<FAssetRegistryModule>("AssetRegistry");
 		const IAssetRegistry& AssetRegistry = AssetRegistryModule.Get();
@@ -162,7 +120,7 @@ UNamingTokens* UNamingTokensEngineSubsystem::GetNamingTokens(const FString& InNa
 						if (!FoundTokens || IsAssetChildOfNativeToken(ClassObjectPath, FoundTokens))
 						{
 							const TSoftClassPtr<UNamingTokens> SoftClassPath(ClassObjectPath.ToString());
-							FoundTokens = LoadNamingToken(SoftClassPath);
+							FoundTokens = LoadNamingToken(SoftClassPath, InNamespace);
 						}
 						else
 						{
@@ -178,6 +136,46 @@ UNamingTokens* UNamingTokensEngineSubsystem::GetNamingTokens(const FString& InNa
 	return FoundTokens;
 }
 
+UNamingTokens* UNamingTokensEngineSubsystem::GetNamingTokensNative(const FString& InNamespace) const
+{
+	// Check cache.
+	UNamingTokens* FoundTokens = GetNamingTokenFromCache(InNamespace, true);
+	if (FoundTokens)
+	{
+		return FoundTokens;
+	}
+
+	const UClass* TargetClass = UNamingTokens::StaticClass();
+
+	// Check native classes.
+	constexpr EClassFlags InvalidClassFlags = CLASS_Deprecated | CLASS_NewerVersionExists | CLASS_Abstract;
+
+	TArray<UClass*> DerivedClasses;
+	GetDerivedClasses(TargetClass, DerivedClasses);
+	for (const UClass* DerivedClass : DerivedClasses)
+	{
+		if (DerivedClass->IsNative()
+			&& DerivedClass->IsChildOf(TargetClass)
+			&& !DerivedClass->HasAnyClassFlags(InvalidClassFlags)
+			&& !DerivedClass->GetName().StartsWith(TEXT("SKEL_"))
+			&& !DerivedClass->GetName().StartsWith(TEXT("REINST_"))
+			&& DerivedClass->GetDefaultObject<UNamingTokens>()->GetNamespace() == InNamespace)
+		{
+			if (!FoundTokens)
+			{
+				FoundTokens = LoadNamingToken(DerivedClass, InNamespace);
+			}
+			else
+			{
+				UE_LOG(LogNamingTokens, Warning, TEXT("Namespace '%s' exists more than once in native class '%s' and '%s'."),
+						*InNamespace, *FoundTokens->GetClass()->GetName(), *DerivedClass->GetName());
+			}
+		}
+	}
+
+	return FoundTokens;
+}
+
 TArray<UNamingTokens*> UNamingTokensEngineSubsystem::GetMultipleNamingTokens(const TArray<FString>& InNamespaces) const
 {
 	TArray<UNamingTokens*> Result;
@@ -218,7 +216,19 @@ FNamingTokenResultData UNamingTokensEngineSubsystem::EvaluateTokenText(const FTe
 	EvaluationData.bForceCaseSensitive = InFilter.bForceCaseSensitive;
 	for (const FString& Namespace : Namespaces)
 	{
-		if (UNamingTokens* NamingTokens = GetNamingTokens(Namespace))
+		DECLARE_DELEGATE_RetVal_OneParam(UNamingTokens*, FGetNamingTokensFunction, const FString&);
+		
+		FGetNamingTokensFunction GetNamingTokensFunction;
+		if (InFilter.bNativeOnly)
+		{
+			GetNamingTokensFunction.BindUObject(this, &UNamingTokensEngineSubsystem::GetNamingTokensNative);
+		}
+		else
+		{
+			GetNamingTokensFunction.BindUObject(this, &UNamingTokensEngineSubsystem::GetNamingTokens);
+		}
+
+		if (UNamingTokens* NamingTokens = GetNamingTokensFunction.Execute(Namespace))
 		{
 			Result = NamingTokens->EvaluateTokenText(ProcessedTokenText, EvaluationData);
 
@@ -364,5 +374,43 @@ FString UNamingTokensEngineSubsystem::GetFormattedTokensStringForDisplay(const F
 
 void UNamingTokensEngineSubsystem::ClearCachedNamingTokens()
 {
+	FScopeLock Lock(&CachedNamingTokensMutex);
 	CachedNamingTokens.Empty();
 }
+
+UNamingTokens* UNamingTokensEngineSubsystem::LoadNamingToken(const TSoftClassPtr<UNamingTokens>& InTokensClass, const FString& InNamespace) const
+{
+	if (const UClass* Class = InTokensClass.LoadSynchronous())
+	{
+		UNamingTokens* NamingTokens = NewObject<UNamingTokens>(GetTransientPackage(), Class);
+		// Initial validation on load. This will log any errors.
+		NamingTokens->Validate();
+		if (ensure(NamingTokens->GetNamespace() == InNamespace))
+		{
+			FScopeLock Lock(&CachedNamingTokensMutex);
+			CachedNamingTokens.Add(InNamespace, NamingTokens);
+			return NamingTokens;
+		}
+	}
+
+	return nullptr;
+}
+
+UNamingTokens* UNamingTokensEngineSubsystem::GetNamingTokenFromCache(const FString& InNamespace, bool bNativeOnly) const
+{
+	FScopeLock Lock(&CachedNamingTokensMutex);
+	if (const TObjectPtr<UNamingTokens>* Tokens = CachedNamingTokens.Find(InNamespace))
+	{
+		if (!bNativeOnly)
+		{
+			return *Tokens;
+		}
+
+		if (bNativeOnly && (*Tokens)->GetClass()->IsNative())
+		{
+			return *Tokens;
+		}
+	}
+
+	return nullptr;
+}
diff --git a/Engine/Plugins/Developer/NamingTokens/Source/NamingTokens/Public/NamingTokensEngineSubsystem.h b/Engine/Plugins/Developer/NamingTokens/Source/NamingTokens/Public/NamingTokensEngineSubsystem.h
--- a/Engine/Plugins/Developer/NamingTokens/Source/NamingTokens/Public/NamingTokensEngineSubsystem.h
+++ b/Engine/Plugins/Developer/NamingTokens/Source/NamingTokens/Public/NamingTokensEngineSubsystem.h
@@ -33,6 +33,10 @@ struct FNamingTokenFilterArgs
 	/** When false, we fall back to case-insensitive if an exact match isn't found. */
 	UPROPERTY(EditAnywhere, BlueprintReadWrite, Category=NamingTokens)
 	bool bForceCaseSensitive = false;
+
+	/** When false, we additionally look for blueprint naming tokens. */
+	UPROPERTY(EditAnywhere, BlueprintReadWrite, Category = NamingTokens)
+	bool bNativeOnly = false;
 };
 
 /**
@@ -55,6 +59,17 @@ public:
 	UFUNCTION(BlueprintCallable, Category = "NamingTokens")
 	UE_API UNamingTokens* GetNamingTokens(const FString& InNamespace) const;
 
+	/**
+	 * Lookup naming tokens given a namespace. This will look first in cached naming tokens,
+	 * then native classes. Call ClearCachedNamingTokens to reset the cache
+	 * and force a full lookup.
+	 *
+	 * @param InNamespace The namespace of the tokens.
+	 * @return The found Naming Tokens object, or nullptr.
+	 */
+	UFUNCTION(BlueprintCallable, Category = "NamingTokens")
+	UE_API UNamingTokens* GetNamingTokensNative(const FString& InNamespace) const;
+
 	/**
 	 * Lookup multiple naming tokens from multiple namespaces.
 	 *
@@ -139,6 +154,10 @@ private:
 	/** Naming tokens currently loaded from assets. */
 	UPROPERTY(Transient)
 	mutable TMap<FString, TObjectPtr<UNamingTokens>> CachedNamingTokens;
+	mutable FCriticalSection CachedNamingTokensMutex;
+
+	UNamingTokens* LoadNamingToken(const TSoftClassPtr<UNamingTokens>& InTokensClass, const FString& InNamespace) const;
+	UNamingTokens* GetNamingTokenFromCache(const FString& InNamespace, bool bNativeOnly = false) const;
 
 	/** Namespaces considered global (don't need to include namespace to access). */
 	TSet<FString> GlobalNamespaces;
diff --git a/Engine/Plugins/VirtualProduction/CaptureManager/CaptureManagerApp/Source/CaptureDataConverter/Private/Nodes/ThirdPartyEncoder/CaptureConvertAudioDataThirdParty.cpp b/Engine/Plugins/VirtualProduction/CaptureManager/CaptureManagerApp/Source/CaptureDataConverter/Private/Nodes/ThirdPartyEncoder/CaptureConvertAudioDataThirdParty.cpp
--- a/Engine/Plugins/VirtualProduction/CaptureManager/CaptureManagerApp/Source/CaptureDataConverter/Private/Nodes/ThirdPartyEncoder/CaptureConvertAudioDataThirdParty.cpp
+++ b/Engine/Plugins/VirtualProduction/CaptureManager/CaptureManagerApp/Source/CaptureDataConverter/Private/Nodes/ThirdPartyEncoder/CaptureConvertAudioDataThirdParty.cpp
@@ -105,30 +105,26 @@ FCaptureConvertAudioDataThirdParty::FResult FCaptureConvertAudioDataThirdParty::
 
 	FString CommandArgs = ThirdPartyEncoder.CommandArguments;
 
-	CallOnGameThread(
-		[&AudioFilePath, &AudioOutputFile, &CommandArgs]()
-		{
-			using namespace UE::CaptureManager;
+	using namespace UE::CaptureManager;
 
-			UNamingTokensEngineSubsystem* NamingTokensSubsystem = GEngine->GetEngineSubsystem<UNamingTokensEngineSubsystem>();
-			check(NamingTokensSubsystem);
+	UNamingTokensEngineSubsystem* NamingTokensSubsystem = GEngine->GetEngineSubsystem<UNamingTokensEngineSubsystem>();
+	check(NamingTokensSubsystem);
 
-			const UCaptureManagerSettings* Settings = GetDefault<UCaptureManagerSettings>();
+	const UCaptureManagerSettings* Settings = GetDefault<UCaptureManagerSettings>();
 
-			FNamingTokenFilterArgs AudioEncoderTokenArgs;
-			const TObjectPtr<const UCaptureManagerAudioEncoderTokens> Tokens = Settings->GetAudioEncoderNamingTokens();
-			check(Tokens);
-			AudioEncoderTokenArgs.AdditionalNamespacesToInclude.Add(Tokens->GetNamespace());
+	FNamingTokenFilterArgs AudioEncoderTokenArgs;
+	const TObjectPtr<const UCaptureManagerAudioEncoderTokens> Tokens = Settings->GetAudioEncoderNamingTokens();
+	check(Tokens);
+	AudioEncoderTokenArgs.AdditionalNamespacesToInclude.Add(Tokens->GetNamespace());
+	AudioEncoderTokenArgs.bNativeOnly = true;
 
-			FStringFormatNamedArguments AudioEncoderFormatArgs;
-			AudioEncoderFormatArgs.Add(Tokens->GetToken(FString(AudioEncoderTokens::InputKey)).Name, WrapInQuotes(AudioFilePath));
-			AudioEncoderFormatArgs.Add(Tokens->GetToken(FString(AudioEncoderTokens::OutputKey)).Name, WrapInQuotes(AudioOutputFile));
-		
-			CommandArgs = FString::Format(*CommandArgs, AudioEncoderFormatArgs);
-			FNamingTokenResultData AudioEncoderCommandResult = NamingTokensSubsystem->EvaluateTokenString(CommandArgs, AudioEncoderTokenArgs);
-			CommandArgs = AudioEncoderCommandResult.EvaluatedText.ToString();
-		}
-	);
+	FStringFormatNamedArguments AudioEncoderFormatArgs;
+	AudioEncoderFormatArgs.Add(Tokens->GetToken(FString(AudioEncoderTokens::InputKey)).Name, WrapInQuotes(AudioFilePath));
+	AudioEncoderFormatArgs.Add(Tokens->GetToken(FString(AudioEncoderTokens::OutputKey)).Name, WrapInQuotes(AudioOutputFile));
+
+	CommandArgs = FString::Format(*CommandArgs, AudioEncoderFormatArgs);
+	FNamingTokenResultData AudioEncoderCommandResult = NamingTokensSubsystem->EvaluateTokenString(CommandArgs, AudioEncoderTokenArgs);
+	CommandArgs = AudioEncoderCommandResult.EvaluatedText.ToString();
 
 	UE_LOG(LogCaptureConvertAudioDataThirdParty, Display, TEXT("Running the command: %s %s"), *ThirdPartyEncoder.Encoder, *CommandArgs);
 
diff --git a/Engine/Plugins/VirtualProduction/CaptureManager/CaptureManagerApp/Source/CaptureDataConverter/Private/Nodes/ThirdPartyEncoder/CaptureConvertVideoDataThirdParty.cpp b/Engine/Plugins/VirtualProduction/CaptureManager/CaptureManagerApp/Source/CaptureDataConverter/Private/Nodes/ThirdPartyEncoder/CaptureConvertVideoDataThirdParty.cpp
--- a/Engine/Plugins/VirtualProduction/CaptureManager/CaptureManagerApp/Source/CaptureDataConverter/Private/Nodes/ThirdPartyEncoder/CaptureConvertVideoDataThirdParty.cpp
+++ b/Engine/Plugins/VirtualProduction/CaptureManager/CaptureManagerApp/Source/CaptureDataConverter/Private/Nodes/ThirdPartyEncoder/CaptureConvertVideoDataThirdParty.cpp
@@ -174,31 +174,25 @@ FCaptureConvertVideoDataThirdParty::FResult FCaptureConvertVideoDataThirdParty::
 
 	FString CommandArgs = ThirdPartyEncoder.CommandArguments;
 
-	CallOnGameThread(
-		[&ConversionArguments, &VideoFilePath, &ImageFilePath, &CommandArgs]()
-		{
-			using namespace UE::CaptureManager;
-
-			UNamingTokensEngineSubsystem* NamingTokensSubsystem = GEngine->GetEngineSubsystem<UNamingTokensEngineSubsystem>();
-			check(NamingTokensSubsystem);
-
-			const UCaptureManagerSettings* Settings = GetDefault<UCaptureManagerSettings>();
-
-			FNamingTokenFilterArgs VideoEncoderTokenArgs;
-			const TObjectPtr<const UCaptureManagerVideoEncoderTokens> Tokens = Settings->GetVideoEncoderNamingTokens();
-			check(Tokens);
-			VideoEncoderTokenArgs.AdditionalNamespacesToInclude.Add(Tokens->GetNamespace());
-
-			FStringFormatNamedArguments VideoEncoderFormatArgs;
-			VideoEncoderFormatArgs.Add(Tokens->GetToken(FString(VideoEncoderTokens::InputKey)).Name, WrapInQuotes(VideoFilePath));
-			VideoEncoderFormatArgs.Add(Tokens->GetToken(FString(VideoEncoderTokens::OutputKey)).Name, WrapInQuotes(ImageFilePath));
-			VideoEncoderFormatArgs.Add(Tokens->GetToken(FString(VideoEncoderTokens::ParamsKey)).Name, ConversionArguments);
-			
-			CommandArgs = FString::Format(*CommandArgs, VideoEncoderFormatArgs);
-			FNamingTokenResultData VideoEncoderCommandResult = NamingTokensSubsystem->EvaluateTokenString(CommandArgs, VideoEncoderTokenArgs);
-			CommandArgs = VideoEncoderCommandResult.EvaluatedText.ToString();
-		}
-	);
+	UNamingTokensEngineSubsystem* NamingTokensSubsystem = GEngine->GetEngineSubsystem<UNamingTokensEngineSubsystem>();
+	check(NamingTokensSubsystem);
+
+	const UCaptureManagerSettings* Settings = GetDefault<UCaptureManagerSettings>();
+
+	FNamingTokenFilterArgs VideoEncoderTokenArgs;
+	const TObjectPtr<const UCaptureManagerVideoEncoderTokens> Tokens = Settings->GetVideoEncoderNamingTokens();
+	check(Tokens);
+	VideoEncoderTokenArgs.AdditionalNamespacesToInclude.Add(Tokens->GetNamespace());
+	VideoEncoderTokenArgs.bNativeOnly = true;
+
+	FStringFormatNamedArguments VideoEncoderFormatArgs;
+	VideoEncoderFormatArgs.Add(Tokens->GetToken(FString(VideoEncoderTokens::InputKey)).Name, WrapInQuotes(VideoFilePath));
+	VideoEncoderFormatArgs.Add(Tokens->GetToken(FString(VideoEncoderTokens::OutputKey)).Name, WrapInQuotes(ImageFilePath));
+	VideoEncoderFormatArgs.Add(Tokens->GetToken(FString(VideoEncoderTokens::ParamsKey)).Name, ConversionArguments);
+
+	CommandArgs = FString::Format(*CommandArgs, VideoEncoderFormatArgs);
+	FNamingTokenResultData VideoEncoderCommandResult = NamingTokensSubsystem->EvaluateTokenString(CommandArgs, VideoEncoderTokenArgs);
+	CommandArgs = VideoEncoderCommandResult.EvaluatedText.ToString();
 	
 	UE_LOG(LogCaptureConvertVideoDataThirdParty, Display, TEXT("Running the command: %s %s"), *ThirdPartyEncoder.Encoder, *CommandArgs);
 
diff --git a/Engine/Plugins/VirtualProduction/CaptureManager/CaptureManagerApp/Source/CaptureManagerEditor/Private/IngestManagement/IngestJobExecutor.cpp b/Engine/Plugins/VirtualProduction/CaptureManager/CaptureManagerApp/Source/CaptureManagerEditor/Private/IngestManagement/IngestJobExecutor.cpp
--- a/Engine/Plugins/VirtualProduction/CaptureManager/CaptureManagerApp/Source/CaptureManagerEditor/Private/IngestManagement/IngestJobExecutor.cpp
+++ b/Engine/Plugins/VirtualProduction/CaptureManager/CaptureManagerApp/Source/CaptureManagerEditor/Private/IngestManagement/IngestJobExecutor.cpp
@@ -194,58 +194,54 @@ TUniquePtr<FIngestJobExecutor::FRunIngestContext> FIngestJobExecutor::RunIngest(
 	UIngestCapability_Options* IngestOptions = NewObject<UIngestCapability_Options>();
 
 	// Naming tokens subsystem consults asset registry so need to run on the game thread
-	CallOnGameThread(
-		[&InJob, &DeviceName, IngestOptions]()
-		{
-			UNamingTokensEngineSubsystem* NamingTokensSubsystem = GEngine->GetEngineSubsystem<UNamingTokensEngineSubsystem>();
-			check(NamingTokensSubsystem);
+	UNamingTokensEngineSubsystem* NamingTokensSubsystem = GEngine->GetEngineSubsystem<UNamingTokensEngineSubsystem>();
+	check(NamingTokensSubsystem);
 
-			FNamingTokenFilterArgs GeneralNamingTokenArgs;
-			const TObjectPtr<const UCaptureManagerGeneralTokens> Tokens = GetDefault<UCaptureManagerSettings>()->GetGeneralNamingTokens();
-			check(Tokens);
+	FNamingTokenFilterArgs GeneralNamingTokenArgs;
+	const TObjectPtr<const UCaptureManagerGeneralTokens> Tokens = GetDefault<UCaptureManagerSettings>()->GetGeneralNamingTokens();
+	check(Tokens);
 
-			GeneralNamingTokenArgs.AdditionalNamespacesToInclude.Add(Tokens->GetNamespace());
+	GeneralNamingTokenArgs.AdditionalNamespacesToInclude.Add(Tokens->GetNamespace());
+	GeneralNamingTokenArgs.bNativeOnly = true;
 
-			const FIngestJob::FSettings Settings = InJob.GetSettings();
-			// If everything is working correctly the upload host name should never be empty.
-			check(!Settings.UploadHostName.IsEmpty());
+	const FIngestJob::FSettings Settings = InJob.GetSettings();
+	// If everything is working correctly the upload host name should never be empty.
+	check(!Settings.UploadHostName.IsEmpty());
 
-			FStringFormatNamedArguments GeneralArgs;
-			{
-				using namespace UE::CaptureManager;
-				GeneralArgs.Add(Tokens->GetToken(FString(GeneralTokens::IdKey)).Name, InJob.GetTakeMetadata().UniqueId);
-				GeneralArgs.Add(Tokens->GetToken(FString(GeneralTokens::DeviceKey)).Name, DeviceName.ToString());
-				GeneralArgs.Add(Tokens->GetToken(FString(GeneralTokens::SlateKey)).Name, InJob.GetTakeMetadata().Slate);
-				GeneralArgs.Add(Tokens->GetToken(FString(GeneralTokens::TakeKey)).Name, InJob.GetTakeMetadata().TakeNumber);
-			}
+	FStringFormatNamedArguments GeneralArgs;
+	{
+		using namespace UE::CaptureManager;
+		GeneralArgs.Add(Tokens->GetToken(FString(GeneralTokens::IdKey)).Name, InJob.GetTakeMetadata().UniqueId);
+		GeneralArgs.Add(Tokens->GetToken(FString(GeneralTokens::DeviceKey)).Name, DeviceName.ToString());
+		GeneralArgs.Add(Tokens->GetToken(FString(GeneralTokens::SlateKey)).Name, InJob.GetTakeMetadata().Slate);
+		GeneralArgs.Add(Tokens->GetToken(FString(GeneralTokens::TakeKey)).Name, InJob.GetTakeMetadata().TakeNumber);
+	}
 
-			FString WorkingDirectory = FString::Format(*Settings.WorkingDirectory, GeneralArgs);
-			FNamingTokenResultData WorkingDirectoryResult = NamingTokensSubsystem->EvaluateTokenString(WorkingDirectory, GeneralNamingTokenArgs);
-			IngestOptions->WorkingDirectory = WorkingDirectoryResult.EvaluatedText.ToString();
-
-			FString DownloadDirectory = FString::Format(*Settings.DownloadFolder, GeneralArgs);
-			FNamingTokenResultData DownloadDirectoryResult = NamingTokensSubsystem->EvaluateTokenString(DownloadDirectory, GeneralNamingTokenArgs);
-			IngestOptions->DownloadDirectory = DownloadDirectoryResult.EvaluatedText.ToString();
-
-			FString FileNamePrefix = FString::Format(*Settings.VideoSettings.FileNamePrefix, GeneralArgs);
-			FNamingTokenResultData VideoPrefixResult = NamingTokensSubsystem->EvaluateTokenString(FileNamePrefix, GeneralNamingTokenArgs);
-			FIngestCapability_VideoOptions VideoOptions;
-			VideoOptions.FileNamePrefix = VideoPrefixResult.EvaluatedText.ToString();
-			VideoOptions.Format = StaticEnum<EOutputImageFormat>()->GetDisplayNameTextByValue(static_cast<std::underlying_type_t<EOutputImageFormat>>(Settings.VideoSettings.Format)).ToString();
-			VideoOptions.PixelFormat = ConvertImagePixelFormat(Settings.VideoSettings.ImagePixelFormat);
-			VideoOptions.Rotation = ConvertImageRotation(Settings.VideoSettings.ImageRotation);
-			IngestOptions->Video = VideoOptions;
-
-			FString AudioPrefix = FString::Format(*Settings.AudioSettings.FileNamePrefix, GeneralArgs);
-			FNamingTokenResultData AudioPrefixResult = NamingTokensSubsystem->EvaluateTokenString(AudioPrefix, GeneralNamingTokenArgs);
-			FIngestCapability_AudioOptions AudioOptions;
-			AudioOptions.FileNamePrefix = AudioPrefixResult.EvaluatedText.ToString();
-			AudioOptions.Format = StaticEnum<EAudioFormat>()->GetDisplayNameTextByValue(static_cast<std::underlying_type_t<EAudioFormat>>(Settings.AudioSettings.Format)).ToString();
-			IngestOptions->Audio = AudioOptions;
-
-			IngestOptions->UploadHostName = Settings.UploadHostName;
-		}
-	);
+	FString WorkingDirectory = FString::Format(*Settings.WorkingDirectory, GeneralArgs);
+	FNamingTokenResultData WorkingDirectoryResult = NamingTokensSubsystem->EvaluateTokenString(WorkingDirectory, GeneralNamingTokenArgs);
+	IngestOptions->WorkingDirectory = WorkingDirectoryResult.EvaluatedText.ToString();
+
+	FString DownloadDirectory = FString::Format(*Settings.DownloadFolder, GeneralArgs);
+	FNamingTokenResultData DownloadDirectoryResult = NamingTokensSubsystem->EvaluateTokenString(DownloadDirectory, GeneralNamingTokenArgs);
+	IngestOptions->DownloadDirectory = DownloadDirectoryResult.EvaluatedText.ToString();
+
+	FString FileNamePrefix = FString::Format(*Settings.VideoSettings.FileNamePrefix, GeneralArgs);
+	FNamingTokenResultData VideoPrefixResult = NamingTokensSubsystem->EvaluateTokenString(FileNamePrefix, GeneralNamingTokenArgs);
+	FIngestCapability_VideoOptions VideoOptions;
+	VideoOptions.FileNamePrefix = VideoPrefixResult.EvaluatedText.ToString();
+	VideoOptions.Format = StaticEnum<EOutputImageFormat>()->GetDisplayNameTextByValue(static_cast<std::underlying_type_t<EOutputImageFormat>>(Settings.VideoSettings.Format)).ToString();
+	VideoOptions.PixelFormat = ConvertImagePixelFormat(Settings.VideoSettings.ImagePixelFormat);
+	VideoOptions.Rotation = ConvertImageRotation(Settings.VideoSettings.ImageRotation);
+	IngestOptions->Video = VideoOptions;
+
+	FString AudioPrefix = FString::Format(*Settings.AudioSettings.FileNamePrefix, GeneralArgs);
+	FNamingTokenResultData AudioPrefixResult = NamingTokensSubsystem->EvaluateTokenString(AudioPrefix, GeneralNamingTokenArgs);
+	FIngestCapability_AudioOptions AudioOptions;
+	AudioOptions.FileNamePrefix = AudioPrefixResult.EvaluatedText.ToString();
+	AudioOptions.Format = StaticEnum<EAudioFormat>()->GetDisplayNameTextByValue(static_cast<std::underlying_type_t<EAudioFormat>>(Settings.AudioSettings.Format)).ToString();
+	IngestOptions->Audio = AudioOptions;
+
+	IngestOptions->UploadHostName = Settings.UploadHostName;
 
 	if (!Device->Implements<ULiveLinkDeviceCapability_Ingest>())
 	{
