Starting removes all jobs before adding them to queue again

diff --git a/Engine/Plugins/VirtualProduction/CaptureManager/CaptureManagerApp/Source/CaptureManagerEditor/Private/UI/SIngestJobProcessor.cpp b/Engine/Plugins/VirtualProduction/CaptureManager/CaptureManagerApp/Source/CaptureManagerEditor/Private/UI/SIngestJobProcessor.cpp
--- a/Engine/Plugins/VirtualProduction/CaptureManager/CaptureManagerApp/Source/CaptureManagerEditor/Private/UI/SIngestJobProcessor.cpp
+++ b/Engine/Plugins/VirtualProduction/CaptureManager/CaptureManagerApp/Source/CaptureManagerEditor/Private/UI/SIngestJobProcessor.cpp
@@ -11,6 +11,8 @@
 
 #define LOCTEXT_NAMESPACE "SIngestJobProcessor"
 
+DEFINE_LOG_CATEGORY_STATIC(LogSIngestJobProcessor, Log, All)
+
 namespace UE::CaptureManager
 {
 
@@ -30,26 +32,56 @@ static const FText StatusText = LOCTEXT("StatusText", "Status");
 
 namespace Private
 {
-void ResetJob(TSharedPtr<FIngestJob> InIngestJob, TSharedPtr<FIngestJobProcessor> InIngestJobProcessor)
+TArray<TSharedRef<FIngestJob>> GetFilteredJobs(const TArray<TSharedRef<FIngestJob>>& InIngestJobs, const TArray<FGuid>& InFilteredJobGuids)
 {
-	if (InIngestJobProcessor->RemoveJobs({ InIngestJob->GetGuid(), }))
+	TArray<TSharedRef<FIngestJob>> FilteredJob;
+	for (const TSharedRef<FIngestJob>& IngestJob : InIngestJobs)
 	{
-		TSharedRef<FIngestJob> NewJob = MakeShared<FIngestJob>(
-			InIngestJob->GetCaptureDeviceId(),
-			InIngestJob->GetTakeId(),
-			InIngestJob->GetTakeMetadata(),
-			InIngestJob->GetPipelineConfig(),
-			InIngestJob->GetSettings()
-		);
-
-		if (!InIngestJobProcessor->AddJobs({ MoveTemp(NewJob), }))
+		if (InFilteredJobGuids.Contains(IngestJob->GetGuid()))
 		{
-			UE_LOG(LogTemp, Warning, TEXT("Failed to add new job entry"));
+			FilteredJob.Add(MakeShared<FIngestJob>(
+				IngestJob->GetCaptureDeviceId(),
+				IngestJob->GetTakeId(),
+				IngestJob->GetTakeMetadata(),
+				IngestJob->GetPipelineConfig(),
+				IngestJob->GetSettings()
+			));
 		}
 	}
-	else
+
+	return FilteredJob;
+}
+
+void ResetJobs(TArray<TSharedRef<FIngestJob>> InIngestJobs, TSharedPtr<FIngestJobProcessor> InIngestJobProcessor)
+{
+	TArray<FGuid> AbortedJobs;
+	Algo::TransformIf(InIngestJobs, AbortedJobs,
+					  [](const TSharedRef<FIngestJob>& Input) -> bool
+					  {
+						  return Input->GetProcessingState() == FIngestJob::EProcessingState::Aborted;
+					  },
+					  [](const TSharedRef<FIngestJob>& Input) -> FGuid
+					  {
+						  return Input->GetGuid();
+					  });
+
+	if (AbortedJobs.IsEmpty())
+	{
+		return;
+	}
+
+	if (!InIngestJobProcessor->RemoveJobs(AbortedJobs))
 	{
-		UE_LOG(LogTemp, Warning, TEXT("Failed to remove old job entry"));
+		UE_LOG(LogSIngestJobProcessor, Warning, TEXT("Failed to remove some jobs from the queue"));
+	}
+
+	TArray<TSharedRef<FIngestJob>> FilteredJobs = GetFilteredJobs(InIngestJobs, AbortedJobs);
+
+	check(!FilteredJobs.IsEmpty()); // FilteredJobs can't be empty if AbortedJobs isn't empty
+
+	if (!InIngestJobProcessor->AddJobs(MoveTemp(FilteredJobs)))
+	{
+		UE_LOG(LogSIngestJobProcessor, Warning, TEXT("Failed to add some jobs to the queue"));
 	}
 }
 }
@@ -306,15 +338,15 @@ bool SIngestJobProcessor::StartButtonIsEnabled() const
 
 FReply SIngestJobProcessor::OnStartButtonClicked()
 {
-	for (const TSharedRef<FIngestJob>& IngestJob : IngestJobs)
-	{
-		if (IngestJob->GetProcessingState() == FIngestJob::EProcessingState::Aborted)
-		{
-			Private::ResetJob(IngestJob, IngestJobProcessor);
-		}
-	}
+	FScopeLock Lock(&CriticalSection);
+
+	Private::ResetJobs(IngestJobs, IngestJobProcessor);
+
+	bool bShouldStartProcessing = !IngestJobs.IsEmpty();
+
+	Lock.Unlock();
 
-	if (!IngestJobs.IsEmpty())
+	if (bShouldStartProcessing)
 	{
 		IngestJobProcessor->StartProcessing();
 	}
@@ -499,7 +531,7 @@ TSharedRef<SWidget> SIngestJobRow::CreateActionWidget()
 
 FReply SIngestJobRow::RetryButtonClicked()
 {
-	Private::ResetJob(IngestJob, IngestJobProcessor);
+	Private::ResetJobs({ IngestJob.ToSharedRef()}, IngestJobProcessor);
 
 	return FReply::Handled();
 }
