[CaptureManager] Fix python interface to use updated device names

diff --git a/Engine/Plugins/VirtualProduction/CaptureManager/CaptureManagerApp/Source/CaptureManagerEditor/Private/Scripting/ULiveLinkHubCaptureDevice.cpp b/Engine/Plugins/VirtualProduction/CaptureManager/CaptureManagerApp/Source/CaptureManagerEditor/Private/Scripting/ULiveLinkHubCaptureDevice.cpp
--- a/Engine/Plugins/VirtualProduction/CaptureManager/CaptureManagerApp/Source/CaptureManagerEditor/Private/Scripting/ULiveLinkHubCaptureDevice.cpp
+++ b/Engine/Plugins/VirtualProduction/CaptureManager/CaptureManagerApp/Source/CaptureManagerEditor/Private/Scripting/ULiveLinkHubCaptureDevice.cpp
@@ -37,7 +37,7 @@ struct ULiveLinkHubCaptureDevice::FImpl
 
 	FImpl();
 
-	TObjectPtr<UIngestCapability_ProcessResult> ImportTake(const FLiveLinkHubTakeMetadata& InTake, 
+	TObjectPtr<UIngestCapability_ProcessResult> IngestTake(const FLiveLinkHubTakeMetadata& InTake, 
 														   const TObjectPtr<const UIngestCapability_Options> InConversionSettings) const;
 	TObjectPtr<UIngestCapability_ProcessResult> DownloadTake(const FLiveLinkHubTakeMetadata& InTake,
 															 const TObjectPtr<const UIngestCapability_Options> InConversionSettings) const;
@@ -181,7 +181,7 @@ void ULiveLinkHubCaptureDevice::FImpl::RemoveDevice()
 	Subsystem->RemoveDevice(BackgroundDevice);
 }
 
-TObjectPtr<UIngestCapability_ProcessResult> ULiveLinkHubCaptureDevice::FImpl::ImportTake(const FLiveLinkHubTakeMetadata& InTake, 
+TObjectPtr<UIngestCapability_ProcessResult> ULiveLinkHubCaptureDevice::FImpl::IngestTake(const FLiveLinkHubTakeMetadata& InTake, 
 																						 TObjectPtr<const UIngestCapability_Options> InConversionSettings) const
 {
 	return RunProcessTake(InTake, MoveTemp(InConversionSettings), EIngestCapability_ProcessConfig::Ingest);
@@ -351,10 +351,10 @@ UIngestCapability_ProcessResult* ULiveLinkHubCaptureDevice::Stop()
 	return ImplPtr->Stop();
 }
 
-UIngestCapability_ProcessResult* ULiveLinkHubCaptureDevice::ImportTake(const FLiveLinkHubTakeMetadata& InTake, const UIngestCapability_Options* InConversionSettings) const
+UIngestCapability_ProcessResult* ULiveLinkHubCaptureDevice::IngestTake(const FLiveLinkHubTakeMetadata& InTake, const UIngestCapability_Options* InConversionSettings) const
 {
 	TObjectPtr<const UIngestCapability_Options> ConversionSettings = InConversionSettings;
-	return ImplPtr->ImportTake(InTake, MoveTemp(ConversionSettings));
+	return ImplPtr->IngestTake(InTake, MoveTemp(ConversionSettings));
 }
 
 UIngestCapability_ProcessResult* ULiveLinkHubCaptureDevice::DownloadTake(const FLiveLinkHubTakeMetadata& InTake, const FString& InDownloadDirectory) const
diff --git a/Engine/Plugins/VirtualProduction/CaptureManager/CaptureManagerApp/Source/CaptureManagerEditor/Private/Scripting/ULiveLinkHubCaptureDevice.h b/Engine/Plugins/VirtualProduction/CaptureManager/CaptureManagerApp/Source/CaptureManagerEditor/Private/Scripting/ULiveLinkHubCaptureDevice.h
--- a/Engine/Plugins/VirtualProduction/CaptureManager/CaptureManagerApp/Source/CaptureManagerEditor/Private/Scripting/ULiveLinkHubCaptureDevice.h
+++ b/Engine/Plugins/VirtualProduction/CaptureManager/CaptureManagerApp/Source/CaptureManagerEditor/Private/Scripting/ULiveLinkHubCaptureDevice.h
@@ -65,7 +65,7 @@ public:
 	UIngestCapability_ProcessResult* Stop();
 
 	UFUNCTION(BlueprintCallable, Category = Ingest)
-	UIngestCapability_ProcessResult* ImportTake(const FLiveLinkHubTakeMetadata& InTake, const UIngestCapability_Options* InConversionSettings) const;
+	UIngestCapability_ProcessResult* IngestTake(const FLiveLinkHubTakeMetadata& InTake, const UIngestCapability_Options* InConversionSettings) const;
 
 	UFUNCTION(BlueprintCallable, Category = Ingest)
 	UIngestCapability_ProcessResult* DownloadTake(const FLiveLinkHubTakeMetadata& InTake, const FString& InDownloadDirectory) const;
diff --git a/Engine/Plugins/VirtualProduction/CaptureManager/CaptureManagerDevices/Source/TakeArchiveIngestDevice/Private/TakeArchiveIngestDevice.h b/Engine/Plugins/VirtualProduction/CaptureManager/CaptureManagerDevices/Source/TakeArchiveIngestDevice/Private/TakeArchiveIngestDevice.h
--- a/Engine/Plugins/VirtualProduction/CaptureManager/CaptureManagerDevices/Source/TakeArchiveIngestDevice/Private/TakeArchiveIngestDevice.h
+++ b/Engine/Plugins/VirtualProduction/CaptureManager/CaptureManagerDevices/Source/TakeArchiveIngestDevice/Private/TakeArchiveIngestDevice.h
@@ -38,7 +38,7 @@ class UTakeArchiveIngestDevice final : public UBaseIngestLiveLinkDevice
 	GENERATED_BODY()
 
 public:
-	UFUNCTION(BlueprintCallable, Category = "FilesystemDeviceSettings")
+	UFUNCTION(BlueprintCallable, Category = "TakeArchiveIngestDeviceSettings")
 	const UTakeArchiveIngestDeviceSettings* GetSettings() const;
 
 	//~ Begin ULiveLinkDevice interface
