Timecode extraction off the game thread

diff --git a/Engine/Plugins/Media/ElectraPlayer/Source/ElectraPlayerRuntime/Private/Runtime/Renderer/RendererVideo.cpp b/Engine/Plugins/Media/ElectraPlayer/Source/ElectraPlayerRuntime/Private/Runtime/Renderer/RendererVideo.cpp
--- a/Engine/Plugins/Media/ElectraPlayer/Source/ElectraPlayerRuntime/Private/Runtime/Renderer/RendererVideo.cpp
+++ b/Engine/Plugins/Media/ElectraPlayer/Source/ElectraPlayerRuntime/Private/Runtime/Renderer/RendererVideo.cpp
@@ -280,8 +280,6 @@ void FElectraRendererVideo::TickInput(FTimespan DeltaTime, FTimespan Timecode)
 	LLM_SCOPE(ELLMTag::ElectraPlayer);
 	CSV_SCOPED_TIMING_STAT(ElectraPlayer, MEDIArendererVideoUE_Tick);
 
-	check(IsInGameThread() || IsInSlateThread());
-
 	// Keep stats up to date
 	SET_DWORD_STAT(STAT_ElectraPlayer_MEDIArendererVideoUE_NumUsedOutputVideoSamples, NumOutputTexturesInUse);
 }
diff --git a/Engine/Plugins/VirtualProduction/CaptureManager/CaptureManagerApp/Source/IngestLiveLinkDevice/IngestLiveLinkDevice.Build.cs b/Engine/Plugins/VirtualProduction/CaptureManager/CaptureManagerApp/Source/IngestLiveLinkDevice/IngestLiveLinkDevice.Build.cs
--- a/Engine/Plugins/VirtualProduction/CaptureManager/CaptureManagerApp/Source/IngestLiveLinkDevice/IngestLiveLinkDevice.Build.cs
+++ b/Engine/Plugins/VirtualProduction/CaptureManager/CaptureManagerApp/Source/IngestLiveLinkDevice/IngestLiveLinkDevice.Build.cs
@@ -35,7 +35,8 @@ public class IngestLiveLinkDevice : ModuleRules
 			"CaptureManagerTakeMetadata",
 			"CaptureManagerSettings",
 			"CoreUObject",
-			"UnrealEd"
+			"UnrealEd",
+			"ElectraPlayerPlugin"
 		});
 
 		CppCompileWarningSettings.UnsafeTypeCastWarningLevel = WarningLevel.Error;
diff --git a/Engine/Plugins/VirtualProduction/CaptureManager/CaptureManagerApp/Source/IngestLiveLinkDevice/Private/Utils/CaptureExtractTimecode.cpp b/Engine/Plugins/VirtualProduction/CaptureManager/CaptureManagerApp/Source/IngestLiveLinkDevice/Private/Utils/CaptureExtractTimecode.cpp
--- a/Engine/Plugins/VirtualProduction/CaptureManager/CaptureManagerApp/Source/IngestLiveLinkDevice/Private/Utils/CaptureExtractTimecode.cpp
+++ b/Engine/Plugins/VirtualProduction/CaptureManager/CaptureManagerApp/Source/IngestLiveLinkDevice/Private/Utils/CaptureExtractTimecode.cpp
@@ -10,14 +10,18 @@
 #include "UObject/Package.h"
 #include "Settings/CaptureManagerSettings.h"
 #include "Logging/LogMacros.h"
+#include "IElectraPlayerPluginModule.h"
+#include "IMediaPlayer.h"
+#include "IMediaOptions.h"
+#include "IMediaEventSink.h"
 
 #define LOCTEXT_NAMESPACE "CaptureExtractTimecode"
 
 DECLARE_LOG_CATEGORY_EXTERN(LogCaptureExtractTimecode, Log, All);
 DEFINE_LOG_CATEGORY(LogCaptureExtractTimecode);
 
-FCaptureExtractVideoTimecode::FCaptureExtractVideoTimecode(const FString& InFilePath) :
-	FilePath(InFilePath)
+FCaptureExtractVideoTimecode::FCaptureExtractVideoTimecode(const FString& InFilePath)
+	: FilePath(InFilePath)
 {}
 
 UE::CaptureManager::FTimecodeAndRateResult FCaptureExtractVideoTimecode::Extract()
@@ -61,202 +65,121 @@ UE::CaptureManager::FTimecodeAndRateResult FCaptureExtractVideoTimecode::Extract
 {
 	using namespace UE::CaptureManager;
 
-	//Some media players require waiting a frame before metadata can be read
-	bGetTimecodeDelay = true;
+	using EventSinkResult = TValueOrError<void, FString>;
+	TPromise<EventSinkResult> EventSinkPromise;
 
-	FTimecodeAndRatePromise Promise;
-	TFuture<FTimecodeAndRateResult> Future = Promise.GetFuture();
+	class FMediaEventSink : public IMediaEventSink
+	{
+	public:
 
-	AsyncTask(ENamedThreads::GameThread, [this, &Promise]
+		FMediaEventSink(TPromise<EventSinkResult>& InPromise)
+			: Promise(InPromise)
 		{
-			UDesiredPlayerMediaSource* MediaSource = NewObject<UDesiredPlayerMediaSource>();
-			MediaSource->SetFilePath(FilePath);
-
-			// Create media player.
-			TStrongObjectPtr<UMediaPlayer> MediaPlayer = TStrongObjectPtr<UMediaPlayer>(
-				NewObject<UMediaPlayer>(GetTransientPackage(), MakeUniqueObjectName(GetTransientPackage(),
-					UMediaPlayer::StaticClass())));
-
-			check(MediaPlayer);
 
-			// Open the media.
-			MediaPlayer->PlayOnOpen = false;
-			FMediaPlayerOptions PlayerOptions;
-			PlayerOptions.SetAllAsOptional();
-			PlayerOptions.InternalCustomOptions.Emplace(MediaPlayerOptionValues::ParseTimecodeInfo(), FVariant());
+		}
 
-			if (!MediaPlayer->OpenSourceWithOptions(MediaSource, PlayerOptions))
+		virtual void ReceiveMediaEvent(EMediaEvent Event) override
+		{
+			if (Event == EMediaEvent::MediaOpened)
 			{
-				Promise.SetValue(MakeError(ECaptureExtractTimecodeError::InternalError));
-				return;
+				Promise.SetValue(MakeValue());
 			}
+			else if (Event == EMediaEvent::MediaOpenFailed)
+			{
+				FString Message = TEXT("Failed to open the media file.");
+				Promise.SetValue(MakeError(MoveTemp(Message)));
+			}
+		}
 
-			Context = { MediaPlayer, Promise };
-
-			// Sometimes it is necessary to wait a couple of ticks for the MediaPlayer to be ready.
-			// Log start time so that we can timeout if initialization is taking too long.
-			MediaPlayerWaitStart = FDateTime::Now();
+	private:
 
-			InitializeClockSink();
-		});
+		TPromise<EventSinkResult>& Promise;
 
-	return Future.Get();
-}
+	} MediaEventSink(EventSinkPromise);
 
-void FCaptureExtractVideoTimecode::TickOutput(FTimespan InDeltaTime, FTimespan InTimespan)
-{
-#if WITH_EDITOR
-	using namespace UE::CaptureManager;
+	IElectraPlayerPluginModule& ElectraModule = FModuleManager::LoadModuleChecked<IElectraPlayerPluginModule>("ElectraPlayerPlugin");
 
-	if (Context.IsSet())
+	TSharedPtr<IMediaPlayer, ESPMode::ThreadSafe> MediaPlayer = ElectraModule.CreatePlayer(MediaEventSink);
+	if (!MediaPlayer)
 	{
-		TStrongObjectPtr<UMediaPlayer>& MediaPlayer = Context.GetValue().MediaPlayer;
+		return MakeError(ECaptureExtractTimecodeError::InternalError);
+	}
 
-		check(MediaPlayer);
+	UDesiredPlayerMediaSource* MediaSource = NewObject<UDesiredPlayerMediaSource>();
+	MediaSource->SetFilePath(FilePath);
 
-		if (MediaPlayer->HasError())
-		{
-			FTimecodeAndRatePromise& Promise = Context.GetValue().ExtractedTimecodeAndRatePromise;
+	FString FileUrl = TEXT("file://") + FilePath;
 
-			Context.Reset();
-			DeinitializeClockSink();
+	FMediaPlayerOptions PlayerOptions;
+	PlayerOptions.SetAllAsOptional();
+	PlayerOptions.InternalCustomOptions.Emplace(MediaPlayerOptionValues::ParseTimecodeInfo(), FVariant());
 
-			Promise.SetValue(MakeError(ECaptureExtractTimecodeError::InternalError));
-			return;
-		}
+	MediaPlayer->Open(FileUrl, MediaSource, &PlayerOptions);
 
-		// Player is not ready if we can't get duration. Delay until duration is not zero.
-		if (bGetTimecodeDelay)
-		{
-			const FTimespan Duration = MediaPlayer->GetDuration();
-			if (Duration == 0)
-			{
-				if ((FDateTime::Now() - MediaPlayerWaitStart).GetSeconds() > TimeoutPeriod)
-				{
-					FTimecodeAndRatePromise& Promise = Context.GetValue().ExtractedTimecodeAndRatePromise;
-					MediaPlayer->Close();
-					Context.Reset();
-					DeinitializeClockSink();
-					Promise.SetValue(MakeValue(FTimecodeAndRate()));
-					return;
-				}
-
-				// Keep delaying
-				return;
-			}
-
-			bGetTimecodeDelay = false;
-		}
-
-		TValueOrError<FTimecode, ECaptureExtractTimecodeError> TimecodeResult = GetTimecode(MediaPlayer);
-		TValueOrError<FFrameRate, ECaptureExtractTimecodeError> TimecodeRateResult = GetTimecodeRate(MediaPlayer);
-		MediaPlayer->Close();
+	FString TimecodeString;
+	FString TimecodeRateString;
 
-		FTimecodeAndRatePromise& Promise = Context.GetValue().ExtractedTimecodeAndRatePromise;
+	FDateTime EventWaitStart = FDateTime::Now();
 
-		Context.Reset();
-		DeinitializeClockSink();
+	TFuture<EventSinkResult> EventSinkFuture = EventSinkPromise.GetFuture();
+	
+	while ((FDateTime::Now() - EventWaitStart).GetSeconds() < TimeoutPeriod)
+	{
+		MediaPlayer->TickInput(0, 0);
 
-		FTimecodeAndRateResult Result = MakeError(ECaptureExtractTimecodeError::TimecodeNotFound);
-		FTimecodeAndRate TimecodeAndRate;
-		if (TimecodeResult.IsValid())
+		if (EventSinkFuture.WaitFor(FTimespan::FromMilliseconds(100.0f))) // Wait 100 milliseconds before invoking TickInput again
 		{
-			TimecodeAndRate.Timecode = TimecodeResult.StealValue();
-
-			if (TimecodeRateResult.IsValid())
+			EventSinkResult Result = EventSinkFuture.Get();
+			if (Result.HasError())
 			{
-				TimecodeAndRate.MaybeTimecodeRate = TimecodeRateResult.StealValue();
+				return MakeError(ECaptureExtractTimecodeError::UnableToOpenMedia);
 			}
 
-			Result = MakeValue(MoveTemp(TimecodeAndRate));
-		}
-
-		Promise.SetValue(MoveTemp(Result));
-	}
-#endif // WITH_EDITOR
-}
-
-void FCaptureExtractVideoTimecode::InitializeClockSink()
-{
-	IMediaModule* MediaModule = FModuleManager::LoadModulePtr<IMediaModule>("Media");
-	if (MediaModule != nullptr)
-	{
-		MediaModule->GetClock().AddSink(this->AsShared());
-	}
-}
-
-void FCaptureExtractVideoTimecode::DeinitializeClockSink()
-{
-	IMediaModule* MediaModule = FModuleManager::GetModulePtr<IMediaModule>("Media");
-	if (MediaModule != nullptr)
-	{
-		MediaModule->GetClock().RemoveSink(this->AsShared());
-	}
-}
-
-TValueOrError<FTimecode, UE::CaptureManager::ECaptureExtractTimecodeError> FCaptureExtractVideoTimecode::GetTimecode(TStrongObjectPtr<UMediaPlayer>& MediaPlayer)
-{
-	using namespace UE::CaptureManager;
+			const FVariant Timecode = MediaPlayer->GetMediaInfo(UMediaPlayer::MediaInfoNameStartTimecodeValue.Resolve());
+			if (!Timecode.IsEmpty())
+			{
+				TimecodeString = Timecode.GetValue<FString>();
 
-	// Timecode from the source. Either player's meta data or media source.
-	const FVariant Value = MediaPlayer->GetMediaInfo(UMediaPlayer::MediaInfoNameStartTimecodeValue.Resolve());
-	FString TimecodeString;
+				const FVariant TimecodeRate = MediaPlayer->GetMediaInfo(UMediaPlayer::MediaInfoNameStartTimecodeFrameRate.Resolve());
+				if (!TimecodeRate.IsEmpty())
+				{
+					TimecodeRateString = TimecodeRate.GetValue<FString>();
+				}
 
-	if (!Value.IsEmpty())
-	{
-		TimecodeString = Value.GetValue<FString>();
-	}
-	else
-	{
-		if (UMediaPlaylist* Playlist = MediaPlayer->GetPlaylist())
-		{
-			if (const UMediaSource* MediaSource = Playlist->Get(MediaPlayer->GetPlaylistIndex()))
-			{
-				TimecodeString = MediaSource->GetMediaOption(UMediaPlayer::MediaInfoNameStartTimecodeValue.Resolve(), FString());
+				break;
 			}
 		}
 	}
 
+	TOptional<FFrameRate> MaybeTimecodeRate;
+
 	if (TimecodeString.IsEmpty())
 	{
 		return MakeError(ECaptureExtractTimecodeError::TimecodeNotFound);
 	}
-	else
+	
+	const TOptional<FTimecode> Timecode = FTimecode::ParseTimecode(*TimecodeString);
+	if (!Timecode.IsSet())
 	{
-		const TOptional<FTimecode> Timecode = FTimecode::ParseTimecode(*TimecodeString);
-		if (!Timecode.IsSet())
+		return MakeError(ECaptureExtractTimecodeError::UnableToParseTimecode);
+	}
+
+	if (!TimecodeRateString.IsEmpty())
+	{
+		TValueOrError<FFrameRate, UE::CaptureManager::ECaptureExtractTimecodeError> Result = ParseTimecodeRate(TimecodeRateString);
+		if (Result.IsValid())
 		{
-			return MakeError(ECaptureExtractTimecodeError::UnableToParseTimecode);
+			MaybeTimecodeRate = Result.GetValue();
 		}
-
-		return MakeValue(Timecode.GetValue());
 	}
+
+	return MakeValue(Timecode.GetValue(), MoveTemp(MaybeTimecodeRate));
 }
 
-TValueOrError<FFrameRate, UE::CaptureManager::ECaptureExtractTimecodeError> FCaptureExtractVideoTimecode::GetTimecodeRate(TStrongObjectPtr<UMediaPlayer>& MediaPlayer)
+TValueOrError<FFrameRate, UE::CaptureManager::ECaptureExtractTimecodeError> FCaptureExtractVideoTimecode::ParseTimecodeRate(FString TimecodeRateString)
 {
 	using namespace UE::CaptureManager;
 
-	// Timecode from the source. Either player's meta data or media source.
-	const FVariant Value = MediaPlayer->GetMediaInfo(UMediaPlayer::MediaInfoNameStartTimecodeFrameRate.Resolve());
-	FString TimecodeRateString;
-
-	if (!Value.IsEmpty())
-	{
-		TimecodeRateString = Value.GetValue<FString>();
-	}
-	else
-	{
-
-		if (UMediaPlaylist* Playlist = MediaPlayer->GetPlaylist())
-		{
-			if (const UMediaSource* MediaSource = Playlist->Get(MediaPlayer->GetPlaylistIndex()))
-			{
-				TimecodeRateString = MediaSource->GetMediaOption(UMediaPlayer::MediaInfoNameStartTimecodeFrameRate.Resolve(), FString());
-			}
-		}
-	}
-
 	if (TimecodeRateString.IsEmpty())
 	{
 		return MakeError(ECaptureExtractTimecodeError::TimecodeRateNotFound);
@@ -460,8 +383,8 @@ UE::CaptureManager::FTimecodeAndRateResult FCaptureExtractVideoTimecode::Extract
 	return Result;
 }
 
-FCaptureExtractAudioTimecode::FCaptureExtractAudioTimecode(const FString& InFilePath) :
-	FilePath(InFilePath)
+FCaptureExtractAudioTimecode::FCaptureExtractAudioTimecode(const FString& InFilePath)
+	: FilePath(InFilePath)
 {}
 
 UE::CaptureManager::FTimecodeAndRateResult FCaptureExtractAudioTimecode::Extract()
diff --git a/Engine/Plugins/VirtualProduction/CaptureManager/CaptureManagerApp/Source/IngestLiveLinkDevice/Public/Utils/CaptureExtractTimecode.h b/Engine/Plugins/VirtualProduction/CaptureManager/CaptureManagerApp/Source/IngestLiveLinkDevice/Public/Utils/CaptureExtractTimecode.h
--- a/Engine/Plugins/VirtualProduction/CaptureManager/CaptureManagerApp/Source/IngestLiveLinkDevice/Public/Utils/CaptureExtractTimecode.h
+++ b/Engine/Plugins/VirtualProduction/CaptureManager/CaptureManagerApp/Source/IngestLiveLinkDevice/Public/Utils/CaptureExtractTimecode.h
@@ -15,8 +15,6 @@
 
 #include "CaptureExtractTimecode.generated.h"
 
-class UMediaPlayer;
-
 namespace UE::CaptureManager
 {
 struct FTimecodeAndRate
@@ -33,31 +31,27 @@ enum ECaptureExtractTimecodeError : int32
 	UnableToParseTimecodeRate = 3,
 	TimecodeRateNotFound = 4,
 	UnhandledMedia = 5,
+	UnableToOpenMedia = 6,
 };
 
 using FTimecodeAndRateResult = TValueOrError<FTimecodeAndRate, ECaptureExtractTimecodeError>;
 }
 
 #if WITH_EDITOR
-
 UCLASS(BlueprintType, MinimalAPI)
 class UDesiredPlayerMediaSource : public UFileMediaSource
 {
 	GENERATED_BODY()
-
 public:
-
 	//~ IMediaOptions interface
-
 	virtual FName GetDesiredPlayerName() const override
 	{
 		return TEXT("ElectraPlayer");
 	}
-
 };
 #endif // WITH_EDITOR
 
-class INGESTLIVELINKDEVICE_API FCaptureExtractVideoTimecode final : public TSharedFromThis<FCaptureExtractVideoTimecode>, public IMediaClockSink
+class INGESTLIVELINKDEVICE_API FCaptureExtractVideoTimecode final : public TSharedFromThis<FCaptureExtractVideoTimecode>
 {
 public:
 
@@ -67,36 +61,14 @@ public:
 
 private:
 
-	using FTimecodeAndRatePromise = TPromise<UE::CaptureManager::FTimecodeAndRateResult>;
-
-	struct FTimecodeExtractionContext
-	{
-		TStrongObjectPtr<UMediaPlayer> MediaPlayer;
-		FTimecodeAndRatePromise& ExtractedTimecodeAndRatePromise;
-	};
-
 	FString FilePath;
 
-	TOptional<FTimecodeExtractionContext> Context;
-
-#if WITH_EDITOR
-	TValueOrError<FTimecode, UE::CaptureManager::ECaptureExtractTimecodeError> GetTimecode(TStrongObjectPtr<UMediaPlayer>& MediaPlayer);
-	TValueOrError<FFrameRate, UE::CaptureManager::ECaptureExtractTimecodeError> GetTimecodeRate(TStrongObjectPtr<UMediaPlayer>& MediaPlayer);
-#endif
-	void InitializeClockSink();
-	void DeinitializeClockSink();
+	TValueOrError<FFrameRate, UE::CaptureManager::ECaptureExtractTimecodeError> ParseTimecodeRate(FString TimecodeRateString);
 
 	UE::CaptureManager::FTimecodeAndRateResult ExtractTimecodeUsingElectraPlayer();
 	UE::CaptureManager::FTimecodeAndRateResult ExtractTimecodeUsingFFProbe(const FString& InEncoderPath);
-
-	bool bGetTimecodeDelay;
+	
 	const int32 TimeoutPeriod = 3;
-	FDateTime MediaPlayerWaitStart;
-
-	/**
-	 * Called from the media clock.
-	 */
-	virtual void TickOutput(FTimespan DeltaTime, FTimespan Timecode) override;
 };
 
 class INGESTLIVELINKDEVICE_API FCaptureExtractAudioTimecode final : public TSharedFromThis<FCaptureExtractAudioTimecode>
