CaptureManager: Fixed crash on start when not connected to a network.

#rnx

diff --git a/Engine/Plugins/VirtualProduction/CaptureManager/CaptureManagerApp/Source/CaptureManagerUnrealEndpoint/Private/CaptureManagerUnrealEndpoint.cpp b/Engine/Plugins/VirtualProduction/CaptureManager/CaptureManagerApp/Source/CaptureManagerUnrealEndpoint/Private/CaptureManagerUnrealEndpoint.cpp
--- a/Engine/Plugins/VirtualProduction/CaptureManager/CaptureManagerApp/Source/CaptureManagerUnrealEndpoint/Private/CaptureManagerUnrealEndpoint.cpp
+++ b/Engine/Plugins/VirtualProduction/CaptureManager/CaptureManagerApp/Source/CaptureManagerUnrealEndpoint/Private/CaptureManagerUnrealEndpoint.cpp
@@ -1,13 +1,12 @@
 // Copyright Epic Games, Inc. All Rights Reserved.
 
 #include "CaptureManagerUnrealEndpoint.h"
+#include "CaptureManagerUnrealEndpointLog.h"
 
 #include "CaptureCommunicationManager.h"
 
 #include <condition_variable>
 
-DEFINE_LOG_CATEGORY_STATIC(LogCaptureManagerUnrealEndpoint, Log, All);
-
 namespace UE::CaptureManager
 {
 
diff --git a/Engine/Plugins/VirtualProduction/CaptureManager/CaptureManagerApp/Source/CaptureManagerUnrealEndpoint/Private/CaptureManagerUnrealEndpointLog.cpp b/Engine/Plugins/VirtualProduction/CaptureManager/CaptureManagerApp/Source/CaptureManagerUnrealEndpoint/Private/CaptureManagerUnrealEndpointLog.cpp
new file mode 100644
--- /dev/null
+++ b/Engine/Plugins/VirtualProduction/CaptureManager/CaptureManagerApp/Source/CaptureManagerUnrealEndpoint/Private/CaptureManagerUnrealEndpointLog.cpp
@@ -0,0 +1,5 @@
+// Copyright Epic Games, Inc. All Rights Reserved.
+
+#include "CaptureManagerUnrealEndpointLog.h"
+
+DEFINE_LOG_CATEGORY(LogCaptureManagerUnrealEndpoint);
diff --git a/Engine/Plugins/VirtualProduction/CaptureManager/CaptureManagerApp/Source/CaptureManagerUnrealEndpoint/Private/CaptureManagerUnrealEndpointLog.h b/Engine/Plugins/VirtualProduction/CaptureManager/CaptureManagerApp/Source/CaptureManagerUnrealEndpoint/Private/CaptureManagerUnrealEndpointLog.h
new file mode 100644
--- /dev/null
+++ b/Engine/Plugins/VirtualProduction/CaptureManager/CaptureManagerApp/Source/CaptureManagerUnrealEndpoint/Private/CaptureManagerUnrealEndpointLog.h
@@ -0,0 +1,7 @@
+// Copyright Epic Games, Inc. All Rights Reserved.
+
+#pragma once
+
+#include "Logging/LogMacros.h"
+
+DECLARE_LOG_CATEGORY_EXTERN(LogCaptureManagerUnrealEndpoint, Log, All);
diff --git a/Engine/Plugins/VirtualProduction/CaptureManager/CaptureManagerApp/Source/CaptureManagerUnrealEndpoint/Private/CaptureManagerUnrealEndpointManager.cpp b/Engine/Plugins/VirtualProduction/CaptureManager/CaptureManagerApp/Source/CaptureManagerUnrealEndpoint/Private/CaptureManagerUnrealEndpointManager.cpp
--- a/Engine/Plugins/VirtualProduction/CaptureManager/CaptureManagerApp/Source/CaptureManagerUnrealEndpoint/Private/CaptureManagerUnrealEndpointManager.cpp
+++ b/Engine/Plugins/VirtualProduction/CaptureManager/CaptureManagerApp/Source/CaptureManagerUnrealEndpoint/Private/CaptureManagerUnrealEndpointManager.cpp
@@ -1,6 +1,7 @@
 // Copyright Epic Games, Inc. All Rights Reserved.
 
 #include "CaptureManagerUnrealEndpointManager.h"
+#include "CaptureManagerUnrealEndpointLog.h"
 
 #include "DiscoveryRequester.h"
 
@@ -101,10 +102,19 @@ void FUnrealEndpointManager::FImpl::Start()
 		return;
 	}
 
+	TUniquePtr<FDiscoveryRequester> NewRequester = FDiscoveryRequester::Create();
+
 	std::lock_guard<std::mutex> LockGuard(Mutex);
 
+	DiscoveryRequester = MoveTemp(NewRequester);
+
+	if (!DiscoveryRequester)
+	{
+		UE_LOG(LogCaptureManagerUnrealEndpoint, Warning, TEXT("Endpoint manager failed to start (Discovery is disabled)"));
+		return;
+	}
+
 	bIsRunning = true;
-	DiscoveryRequester = MakeUnique<FDiscoveryRequester>();
 
 	EndpointFoundDelegateHandle = DiscoveryRequester->ClientFound().AddLambda(
 		[this](const FDiscoveredClient& InDiscoveredClient)
diff --git a/Engine/Plugins/VirtualProduction/CaptureManager/CaptureManagerApp/Source/CaptureManagerUnrealEndpoint/Private/DiscoveryRequester.cpp b/Engine/Plugins/VirtualProduction/CaptureManager/CaptureManagerApp/Source/CaptureManagerUnrealEndpoint/Private/DiscoveryRequester.cpp
--- a/Engine/Plugins/VirtualProduction/CaptureManager/CaptureManagerApp/Source/CaptureManagerUnrealEndpoint/Private/DiscoveryRequester.cpp
+++ b/Engine/Plugins/VirtualProduction/CaptureManager/CaptureManagerApp/Source/CaptureManagerUnrealEndpoint/Private/DiscoveryRequester.cpp
@@ -12,11 +12,7 @@
 #include "CaptureUtilsModule.h"
 #include "Network/NetworkMisc.h"
 
-#include "Misc/CoreDelegates.h"
-
-#include "HAL/PlatformTime.h"
-
-DEFINE_LOG_CATEGORY(LogLiveLinkHubDiscovery)
+DEFINE_LOG_CATEGORY(LogLiveLinkHubDiscovery);
 
 namespace UE::CaptureManager
 {
@@ -101,7 +97,7 @@ struct FDiscoveryRequester::FImpl
 {
 	static TSharedRef<FCaptureTimerManager> GetTimerManager();
 
-	FImpl();
+	static TUniquePtr<FImpl> Create(FString InLocalHostName);
 	~FImpl();
 
 	void OnTick();
@@ -113,11 +109,9 @@ struct FDiscoveryRequester::FImpl
 	void StartDiscoveryRequester();
 	TArray<FDiscoveredClient> GetDiscoveredClients() const;
 
-	FString LocalHostIPAddress;
 	FString HostName;
 
 	static constexpr float DiscoveryRequestPeriod = 3.0f;
-	float TimeSinceLastDiscoveryRequest = 0.0f;
 
 	FCaptureTimerManager::FTimerHandle TickerHandle;
 	TSharedPtr<FMessageEndpoint> MessageEndpoint;
@@ -125,9 +119,9 @@ struct FDiscoveryRequester::FImpl
 	FClientFound ClientFoundDelegate;
 	FClientLost ClientLostDelegate;
 
-	std::atomic_bool bStarted = false;
-
 private:
+	explicit FImpl(FString InLocalHostName);
+
 	void RemoveStaleClients();
 	mutable FCriticalSection Mutex;
 	TArray<FDiscoveredClient> KnownClients;
@@ -142,32 +136,37 @@ TSharedRef<FCaptureTimerManager> FDiscoveryRequester::FImpl::GetTimerManager()
 	return CaptureUtils.GetTimerManager();
 }
 
-FDiscoveryRequester::FImpl::FImpl()
-	: bStarted(false)
-	, TimerManager(GetTimerManager())
+TUniquePtr<FDiscoveryRequester::FImpl> FDiscoveryRequester::FImpl::Create(FString InLocalHostName)
 {
-	TOptional<FString> IPAddress = GetLocalIpAddress();
-	TOptional<FString> HostNameOpt = GetLocalHostName();
+	TUniquePtr<FImpl> Impl = TUniquePtr<FImpl>(new FImpl(MoveTemp(InLocalHostName)));
 
-	if (IPAddress.IsSet() && HostNameOpt.IsSet())
-	{
-		LocalHostIPAddress = MoveTemp(IPAddress.GetValue());
-		HostName = MoveTemp(HostNameOpt.GetValue());
+	// Message endpoint must be created on the game thread
+	check(IsInGameThread());
 
-		MessageEndpoint = FMessageEndpoint::Builder("DiscoveryRequester")
-			.Handling<FDiscoveryResponse>(this, &FImpl::HandleDiscoveryResponse)
-			.ReceivingOnAnyThread()
-			;
-	}
-	else
+	TSharedPtr<FMessageEndpoint> MessageEndpoint = FMessageEndpoint::Builder("DiscoveryRequester")
+		.Handling<FDiscoveryResponse>(Impl.Get(), &FImpl::HandleDiscoveryResponse)
+		.ReceivingOnAnyThread()
+		;
+
+	if (!MessageEndpoint)
 	{
 		UE_LOG(
 			LogLiveLinkHubDiscovery,
 			Warning,
-			TEXT("Failed to determine the IP address or Host Name for the discovery requester, capture source discovery has been disabled")
+			TEXT("Failed to create message endpoint, discovery will be disabled")
 		);
+
+		return nullptr;
 	}
 
+	Impl->MessageEndpoint = MoveTemp(MessageEndpoint);
+	return Impl;
+}
+
+FDiscoveryRequester::FImpl::FImpl(FString InLocalHostName) :
+	HostName(MoveTemp(InLocalHostName)),
+	TimerManager(GetTimerManager())
+{
 }
 
 FDiscoveryRequester::FImpl::~FImpl()
@@ -219,33 +218,29 @@ void FDiscoveryRequester::FImpl::RemoveStaleClients()
 
 TArray<FDiscoveredClient> FDiscoveryRequester::GetDiscoveredClients() const
 {
-	return ImplPtr->GetDiscoveredClients();
+	return Impl->GetDiscoveredClients();
 }
 
 FDiscoveryRequester::FClientFound& FDiscoveryRequester::ClientFound()
 {
-	return ImplPtr->ClientFoundDelegate;
+	return Impl->ClientFoundDelegate;
 }
 
 FDiscoveryRequester::FClientLost& FDiscoveryRequester::ClientLost()
 {
-	return ImplPtr->ClientLostDelegate;
+	return Impl->ClientLostDelegate;
 }
 
 void FDiscoveryRequester::FImpl::OnTick()
 {
 	check(!IsInGameThread());
 
-	if (!bStarted)
-	{
-		return;
-	}
-
 	RemoveStaleClients();
 
 	FDiscoveryRequest* Message = FMessageEndpoint::MakeMessage<FDiscoveryRequest>();
 	Message->HostName = HostName;
 
+	check(MessageEndpoint);
 	MessageEndpoint->Publish(Message);
 }
 
@@ -291,13 +286,46 @@ void FDiscoveryRequester::FImpl::HandleDiscoveryResponse(
 
 void FDiscoveryRequester::FImpl::StartDiscoveryRequester()
 {
-	bStarted = true;
-
 	TickerHandle = TimerManager->AddTimer(FTimerDelegate::CreateRaw(this, &FImpl::OnTick), DiscoveryRequestPeriod, true);
 }
 
-FDiscoveryRequester::FDiscoveryRequester()
-	: ImplPtr(MakeUnique<FImpl>())
+TUniquePtr<FDiscoveryRequester> FDiscoveryRequester::Create()
+{
+	TOptional<FString> LocalHostName = GetLocalHostName();
+
+	if (!LocalHostName.IsSet())
+	{
+		UE_LOG(
+			LogLiveLinkHubDiscovery,
+			Warning,
+			TEXT("Failed to determine local host name, discovery will be disabled")
+		);
+		return nullptr;
+	}
+
+	if (LocalHostName.GetValue().IsEmpty())
+	{
+		UE_LOG(
+			LogLiveLinkHubDiscovery,
+			Warning,
+			TEXT("The local host name was invalid (empty), discovery will be disabled")
+		);
+		return nullptr;
+	}
+
+	// We construct the Impl in this way to avoid constructing the requester into an invalid state if the message endpoint build fails.
+	TUniquePtr<FImpl> Impl = FImpl::Create(MoveTemp(*LocalHostName));
+
+	if (!Impl)
+	{
+		return nullptr;
+	}
+
+	return TUniquePtr<FDiscoveryRequester>(new FDiscoveryRequester(MoveTemp(Impl)));
+}
+
+FDiscoveryRequester::FDiscoveryRequester(TUniquePtr<FImpl> InImpl)
+	: Impl(MoveTemp(InImpl))
 {
 }
 
@@ -305,7 +333,7 @@ FDiscoveryRequester::~FDiscoveryRequester() = default;
 
 void FDiscoveryRequester::Start()
 {
-	ImplPtr->StartDiscoveryRequester();
+	Impl->StartDiscoveryRequester();
 }
 
 } // namespace UE::CaptureManager
diff --git a/Engine/Plugins/VirtualProduction/CaptureManager/CaptureManagerApp/Source/CaptureManagerUnrealEndpoint/Private/DiscoveryRequester.h b/Engine/Plugins/VirtualProduction/CaptureManager/CaptureManagerApp/Source/CaptureManagerUnrealEndpoint/Private/DiscoveryRequester.h
--- a/Engine/Plugins/VirtualProduction/CaptureManager/CaptureManagerApp/Source/CaptureManagerUnrealEndpoint/Private/DiscoveryRequester.h
+++ b/Engine/Plugins/VirtualProduction/CaptureManager/CaptureManagerApp/Source/CaptureManagerUnrealEndpoint/Private/DiscoveryRequester.h
@@ -4,7 +4,7 @@
 
 #include "MessageEndpoint.h"
 
-DECLARE_LOG_CATEGORY_EXTERN(LogLiveLinkHubDiscovery, Log, All)
+DECLARE_LOG_CATEGORY_EXTERN(LogLiveLinkHubDiscovery, Log, All);
 
 namespace UE::CaptureManager
 {
@@ -54,7 +54,7 @@ public:
 	DECLARE_TS_MULTICAST_DELEGATE_OneParam(FClientFound, const FDiscoveredClient& InClient);
 	DECLARE_TS_MULTICAST_DELEGATE_OneParam(FClientLost, const FGuid& InClientId);
 
-	FDiscoveryRequester();
+	static TUniquePtr<FDiscoveryRequester> Create();
 	~FDiscoveryRequester();
 
 	void Start();
@@ -65,7 +65,10 @@ public:
 
 private:
 	struct FImpl;
-	TUniquePtr<FImpl> ImplPtr;
+
+	explicit FDiscoveryRequester(TUniquePtr<FImpl> InImpl);
+
+	TUniquePtr<FImpl> Impl;
 };
 
 } // namespace UE::CaptureManager
