CaptureManager: Fixed ffprobe timecode always coming in as drop frame.

#rnx

diff --git a/Engine/Plugins/VirtualProduction/CaptureManager/CaptureManagerApp/Source/IngestLiveLinkDevice/Private/Utils/CaptureExtractTimecode.cpp b/Engine/Plugins/VirtualProduction/CaptureManager/CaptureManagerApp/Source/IngestLiveLinkDevice/Private/Utils/CaptureExtractTimecode.cpp
--- a/Engine/Plugins/VirtualProduction/CaptureManager/CaptureManagerApp/Source/IngestLiveLinkDevice/Private/Utils/CaptureExtractTimecode.cpp
+++ b/Engine/Plugins/VirtualProduction/CaptureManager/CaptureManagerApp/Source/IngestLiveLinkDevice/Private/Utils/CaptureExtractTimecode.cpp
@@ -17,8 +17,7 @@
 
 #define LOCTEXT_NAMESPACE "CaptureExtractTimecode"
 
-DECLARE_LOG_CATEGORY_EXTERN(LogCaptureExtractTimecode, Log, All);
-DEFINE_LOG_CATEGORY(LogCaptureExtractTimecode);
+DEFINE_LOG_CATEGORY_STATIC(LogCaptureExtractTimecode, Log, All);
 
 FCaptureExtractVideoTimecode::FCaptureExtractVideoTimecode(const FString& InFilePath)
 	: FilePath(InFilePath)
@@ -345,7 +344,8 @@ UE::CaptureManager::FTimecodeAndRateResult FCaptureExtractVideoTimecode::Extract
 					int32 Seconds = FCString::Atoi(*TimecodePartsStr[2]);
 					int32 Frames = FCString::Atoi(*TimecodePartsStr[3]);
 
-					Timecode = FTimecode(Hours, Minutes, Seconds, Frames, true);
+					constexpr bool bIsDropFrame = false;
+					Timecode = FTimecode(Hours, Minutes, Seconds, Frames, bIsDropFrame);
 
 					bFoundTimecode = true;
 					continue;
