Fixing crash when disconnect happens

* All the issues below are a consequence of a disconnect that happened at the different time
* Check the tickets for details

diff --git a/Engine/Plugins/VirtualProduction/CaptureManager/CaptureManagerApp/Source/CaptureManagerEditor/Private/IngestManagement/IngestJobExecutor.cpp b/Engine/Plugins/VirtualProduction/CaptureManager/CaptureManagerApp/Source/CaptureManagerEditor/Private/IngestManagement/IngestJobExecutor.cpp
--- a/Engine/Plugins/VirtualProduction/CaptureManager/CaptureManagerApp/Source/CaptureManagerEditor/Private/IngestManagement/IngestJobExecutor.cpp
+++ b/Engine/Plugins/VirtualProduction/CaptureManager/CaptureManagerApp/Source/CaptureManagerEditor/Private/IngestManagement/IngestJobExecutor.cpp
@@ -197,6 +197,8 @@ TUniquePtr<FIngestJobExecutor::FRunIngestContext> FIngestJobExecutor::RunIngest(
 	UNamingTokensEngineSubsystem* NamingTokensSubsystem = GEngine->GetEngineSubsystem<UNamingTokensEngineSubsystem>();
 	check(NamingTokensSubsystem);
 
+	NamingTokensSubsystem->SetCacheEnabled(false);
+
 	FNamingTokenFilterArgs GeneralNamingTokenArgs;
 	const TObjectPtr<const UCaptureManagerGeneralTokens> Tokens = GetDefault<UCaptureManagerSettings>()->GetGeneralNamingTokens();
 	check(Tokens);
diff --git a/Engine/Plugins/VirtualProduction/CaptureManager/CaptureManagerApp/Source/CaptureManagerEditor/Private/IngestManagement/IngestJobProcessor.cpp b/Engine/Plugins/VirtualProduction/CaptureManager/CaptureManagerApp/Source/CaptureManagerEditor/Private/IngestManagement/IngestJobProcessor.cpp
--- a/Engine/Plugins/VirtualProduction/CaptureManager/CaptureManagerApp/Source/CaptureManagerEditor/Private/IngestManagement/IngestJobProcessor.cpp
+++ b/Engine/Plugins/VirtualProduction/CaptureManager/CaptureManagerApp/Source/CaptureManagerEditor/Private/IngestManagement/IngestJobProcessor.cpp
@@ -17,7 +17,8 @@ namespace Private
 
 	static void StopUnrealEndpointConnections()
 	{
-		FCaptureManagerUnrealEndpointModule* EndpointModule = FModuleManager::GetModulePtr<FCaptureManagerUnrealEndpointModule>("Settings");
+		FCaptureManagerUnrealEndpointModule* EndpointModule = 
+			FModuleManager::GetModulePtr<FCaptureManagerUnrealEndpointModule>("CaptureManagerUnrealEndpoint");
 		if (!EndpointModule)
 		{
 			return;
@@ -108,6 +109,8 @@ void FIngestJobProcessor::StartProcessing()
 			}
 		);
 
+		++NumExecutorsRunning;
+
 		{
 			FScopeLock Lock(&CriticalSection);
 			Executors.Emplace(
@@ -119,8 +122,6 @@ void FIngestJobProcessor::StartProcessing()
 				)
 			);
 		}
-
-		++NumExecutorsRunning;
 	}
 }
 
diff --git a/Engine/Plugins/VirtualProduction/CaptureManager/CaptureManagerApp/Source/CaptureManagerSettings/Private/Settings/CaptureManagerSettings.cpp b/Engine/Plugins/VirtualProduction/CaptureManager/CaptureManagerApp/Source/CaptureManagerSettings/Private/Settings/CaptureManagerSettings.cpp
--- a/Engine/Plugins/VirtualProduction/CaptureManager/CaptureManagerApp/Source/CaptureManagerSettings/Private/Settings/CaptureManagerSettings.cpp
+++ b/Engine/Plugins/VirtualProduction/CaptureManager/CaptureManagerApp/Source/CaptureManagerSettings/Private/Settings/CaptureManagerSettings.cpp
@@ -5,6 +5,8 @@
 #include "Network/NetworkMisc.h"
 #include "Misc/Paths.h"
 
+#include "Async/TaskGraphInterfaces.h"
+
 namespace UE::CaptureManager::Private
 {
 constexpr FStringView DefaultWorkingDirectory = TEXTVIEW("CaptureManager/Working/{yyyy}{mm}{dd}");
@@ -50,6 +52,20 @@ void UCaptureManagerSettings::PostEditChangeProperty(struct FPropertyChangedEven
 	InitializeValuesIfNotSet();
 }
 
+void UCaptureManagerSettings::PostInitProperties()
+{
+	Super::PostInitProperties();
+
+	FProperty* Property = GetClass()->FindPropertyByName(GET_MEMBER_NAME_CHECKED(UCaptureManagerSettings, NumIngestExecutors));
+	check(Property);
+
+	FString MaxNumberOfExecutorsString = Property->GetMetaData(TEXT("ClampMax"));
+	int32 MaxNumberOfExecutors;
+	TTypeFromString<int32>::FromString(MaxNumberOfExecutors, *MaxNumberOfExecutorsString);
+
+	NumIngestExecutors = FMath::Min(NumIngestExecutors, MaxNumberOfExecutors);
+}
+
 void UCaptureManagerSettings::InitializeValuesIfNotSet()
 {
 	using namespace UE::CaptureManager;
diff --git a/Engine/Plugins/VirtualProduction/CaptureManager/CaptureManagerApp/Source/CaptureManagerSettings/Public/Settings/CaptureManagerSettings.h b/Engine/Plugins/VirtualProduction/CaptureManager/CaptureManagerApp/Source/CaptureManagerSettings/Public/Settings/CaptureManagerSettings.h
--- a/Engine/Plugins/VirtualProduction/CaptureManager/CaptureManagerApp/Source/CaptureManagerSettings/Public/Settings/CaptureManagerSettings.h
+++ b/Engine/Plugins/VirtualProduction/CaptureManager/CaptureManagerApp/Source/CaptureManagerSettings/Public/Settings/CaptureManagerSettings.h
@@ -80,13 +80,14 @@ public:
 	UPROPERTY(config, EditAnywhere, Category = "Upload")
 	FString DefaultUploadHostName;
 
-	/** The number of jobs to run in parallel. Requires a restart of Live Link Hub (Warning: Setting this too high could degrade performance) */
-	UPROPERTY(config, EditAnywhere, Category = "General", meta = (DisplayName = "Parallel Jobs", UIMin = 1, ClampMin = 1, UIMax = 64, ClampMax = 64))
+	/** The number of jobs to run in parallel. Requires a restart of Live Link Hub (Warning: Setting this too high could cause issues) */
+	UPROPERTY(config, EditAnywhere, Category = "General", meta = (DisplayName = "Parallel Jobs", UIMin = 1, ClampMin = 1, UIMax = 8, ClampMax = 8))
 	int32 NumIngestExecutors = 2;
 
 private:
 
 	virtual void PostEditChangeProperty(struct FPropertyChangedEvent& InPropertyChangedEvent) override;
+	virtual void PostInitProperties() override;
 
 	void InitializeValuesIfNotSet();
 
diff --git a/Engine/Plugins/VirtualProduction/CaptureManager/CaptureManagerApp/Source/CaptureManagerUnrealEndpoint/Private/CaptureCommunicationManager.cpp b/Engine/Plugins/VirtualProduction/CaptureManager/CaptureManagerApp/Source/CaptureManagerUnrealEndpoint/Private/CaptureCommunicationManager.cpp
--- a/Engine/Plugins/VirtualProduction/CaptureManager/CaptureManagerApp/Source/CaptureManagerUnrealEndpoint/Private/CaptureCommunicationManager.cpp
+++ b/Engine/Plugins/VirtualProduction/CaptureManager/CaptureManagerApp/Source/CaptureManagerUnrealEndpoint/Private/CaptureCommunicationManager.cpp
@@ -66,6 +66,11 @@ FCommunicationManager::~FCommunicationManager()
 
 void FCommunicationManager::Connect(FMessageAddress InAddress, FString InIpAddress, uint16 InPort)
 {
+	if (Messenger && Messenger->IsConnected())
+	{
+		return;
+	}
+
 	Messenger->SetAddress(MoveTemp(InAddress));
 	Messenger->Connect(
 		FConnectStarter::FConnectHandler::CreateLambda(
@@ -89,21 +94,23 @@ void FCommunicationManager::Connect(FMessageAddress InAddress, FString InIpAddre
 
 void FCommunicationManager::Disconnect()
 {
-	if (Messenger)
+	if (Messenger && !Messenger->IsConnected())
 	{
-		Messenger->SetDisconnectHandler(
-			FConnectStarter::FDisconnectHandler::CreateLambda(
-				[this]()
-				{
-					ConnectionStateChangedDelegate.ExecuteIfBound(Messenger->GetAddress(), false);
-				}
-			)
-		);
-		Messenger->Disconnect();
-
-		FScopeLock Lock(&Mutex);
-		Uploaders.Empty();
+		return;
 	}
+
+	Messenger->SetDisconnectHandler(
+		FConnectStarter::FDisconnectHandler::CreateLambda(
+		[this]()
+		{
+			ConnectionStateChangedDelegate.ExecuteIfBound(Messenger->GetAddress(), false);
+		}
+	)
+	);
+	Messenger->Disconnect();
+
+	FScopeLock Lock(&Mutex);
+	Uploaders.Empty();
 }
 
 TWeakPtr<FUploader> FCommunicationManager::GetUploader(const FGuid& InCaptureSourceId,
@@ -190,7 +197,12 @@ void FCommunicationManager::OnUploadState(const FGuid& InCaptureSourceId, const
 	if (TSharedPtr<FUploader>* Context = Uploaders.Find(InCaptureSourceId))
 	{
 		check(*Context);
-		(*Context)->UploadUpdate.ExecuteIfBound(InTakeUploadId, InProgress);
+
+		// Copying the delegate to avoid calling it while holding a mutex
+		FUploader::FUploadStateUpdate UploadUpdate = (*Context)->UploadUpdate;
+		Lock.Unlock();
+
+		UploadUpdate.ExecuteIfBound(InTakeUploadId, InProgress);
 	}
 }
 
@@ -201,6 +213,11 @@ void FCommunicationManager::OnUploadFinished(const FGuid& InCaptureSourceId, con
 	if (TSharedPtr<FUploader>* Context = Uploaders.Find(InCaptureSourceId))
 	{
 		check(*Context);
-		(*Context)->UploadFinished.ExecuteIfBound(InTakeUploadId, InMessage, InCode);
+
+		// Copying the delegate to avoid calling it while holding a mutex
+		FUploader::FUploadFinished UploadFinished = (*Context)->UploadFinished;
+		Lock.Unlock();
+
+		UploadFinished.ExecuteIfBound(InTakeUploadId, InMessage, InCode);
 	}
 }
diff --git a/Engine/Plugins/VirtualProduction/CaptureManager/CaptureManagerApp/Source/CaptureManagerUnrealEndpoint/Private/CaptureManagerUnrealEndpoint.cpp b/Engine/Plugins/VirtualProduction/CaptureManager/CaptureManagerApp/Source/CaptureManagerUnrealEndpoint/Private/CaptureManagerUnrealEndpoint.cpp
--- a/Engine/Plugins/VirtualProduction/CaptureManager/CaptureManagerApp/Source/CaptureManagerUnrealEndpoint/Private/CaptureManagerUnrealEndpoint.cpp
+++ b/Engine/Plugins/VirtualProduction/CaptureManager/CaptureManagerApp/Source/CaptureManagerUnrealEndpoint/Private/CaptureManagerUnrealEndpoint.cpp
@@ -10,11 +10,13 @@
 namespace UE::CaptureManager
 {
 
-struct FUnrealEndpoint::FImpl
+struct FUnrealEndpoint::FImpl : public TSharedFromThis<FUnrealEndpoint::FImpl>
 {
 	FImpl(FUnrealEndpointInfo EndpointInfo);
 	~FImpl();
 
+	void SetupConnectionDelegate();
+
 	void StartConnection();
 	void StopConnection();
 
@@ -44,8 +46,9 @@ struct FUnrealEndpoint::FImpl
 };
 
 FUnrealEndpoint::FUnrealEndpoint(FUnrealEndpointInfo EndpointInfo) :
-	Impl(MakePimpl<FImpl>(EndpointInfo))
+	Impl(MakeShared<FImpl>(EndpointInfo))
 {
+	Impl->SetupConnectionDelegate();
 }
 
 FUnrealEndpoint::~FUnrealEndpoint() = default;
@@ -84,24 +87,28 @@ FUnrealEndpoint::FImpl::FImpl(FUnrealEndpointInfo InEndpointInfo) :
 	EndpointInfo(MoveTemp(InEndpointInfo)),
 	bConnectionIsStarted(false),
 	bIsConnected(false)
+{
+}
+
+FUnrealEndpoint::FImpl::~FImpl() = default;
+
+void FUnrealEndpoint::FImpl::SetupConnectionDelegate()
 {
 	CommunicationManager.ConnectionStateChanged().BindLambda(
-		[this](const FMessageAddress& InMessageAddress, bool bInIsConnected)
+		[This = AsShared()](const FMessageAddress& InMessageAddress, bool bInIsConnected)
 		{
-			if (InMessageAddress == EndpointInfo.MessageAddress)
+			if (InMessageAddress == This->EndpointInfo.MessageAddress)
 			{
 				{
-					std::lock_guard<std::mutex> LockGuard(CondMtx);
-					bIsConnected = bInIsConnected;
+					std::lock_guard<std::mutex> LockGuard(This->CondMtx);
+					This->bIsConnected = bInIsConnected;
 				}
-				CondVar.notify_one();
+				This->CondVar.notify_one();
 			}
 		}
 	);
 }
 
-FUnrealEndpoint::FImpl::~FImpl() = default;
-
 void FUnrealEndpoint::FImpl::StartConnection()
 {
 	if (bConnectionIsStarted)
@@ -135,15 +142,15 @@ bool FUnrealEndpoint::FImpl::WaitForConnectionState(const FUnrealEndpoint::EConn
 	bool bWaitSuccess = CondVar.wait_for(
 		Lock,
 		std::chrono::milliseconds(TimeoutMs),
-		[this, InConnectionState]() -> bool
+		[This = AsShared(), InConnectionState]() -> bool
 		{
 			switch (InConnectionState)
 			{
 				case EConnectionState::Connected:
-					return bIsConnected;
+					return This->bIsConnected;
 
 				case EConnectionState::Disconnected:
-					return !bIsConnected;
+					return !This->bIsConnected;
 
 				default:
 					check(false);
@@ -168,12 +175,12 @@ bool FUnrealEndpoint::FImpl::AddTakeUploadTask(TUniquePtr<FTakeUploadTask> InTak
 
 	// Note: Only the first call to GetUploader will register these callbacks
 	FUploader::FUploadStateUpdate UpdateCallback = FUploader::FUploadStateUpdate::CreateLambda(
-		[this](const FGuid& InTakeUploadID, const double InProgress)
+		[This = AsShared()](const FGuid& InTakeUploadID, const double InProgress)
 		{
 			// Find the take upload task and update it
-			FScopeLock ScopeLock(&CriticalSection);
+			FScopeLock ScopeLock(&This->CriticalSection);
 
-			const int32 TaskIndex = TakeUploadTasks.IndexOfByPredicate(
+			const int32 TaskIndex = This->TakeUploadTasks.IndexOfByPredicate(
 				[&InTakeUploadID](const TSharedRef<FTakeUploadTask>& InTask)
 				{
 					return InTask->GetTaskID() == InTakeUploadID;
@@ -183,7 +190,7 @@ bool FUnrealEndpoint::FImpl::AddTakeUploadTask(TUniquePtr<FTakeUploadTask> InTak
 			if (TaskIndex != INDEX_NONE)
 			{
 				// Grab a strong reference before releasing the lock
-				TSharedRef<FTakeUploadTask> ThisTask = TakeUploadTasks[TaskIndex];
+				TSharedRef<FTakeUploadTask> ThisTask = This->TakeUploadTasks[TaskIndex];
 				ScopeLock.Unlock();
 
 				ThisTask->Progressed().ExecuteIfBound(InProgress);
@@ -193,12 +200,12 @@ bool FUnrealEndpoint::FImpl::AddTakeUploadTask(TUniquePtr<FTakeUploadTask> InTak
 
 	// Note: Only the first call to GetUploader will register these callbacks
 	FUploader::FUploadFinished UpdateFinishedCallback = FUploader::FUploadFinished::CreateLambda(
-		[this](const FGuid& InTakeUploadID, const FString& InMessage, int32 InCode)
+		[This = AsShared()](const FGuid& InTakeUploadID, const FString& InMessage, int32 InCode)
 		{
-			FScopeLock ScopeLock(&CriticalSection);
+			FScopeLock ScopeLock(&This->CriticalSection);
 
 			// Find the take upload task and update it
-			const int32 TaskIndex = TakeUploadTasks.IndexOfByPredicate(
+			const int32 TaskIndex = This->TakeUploadTasks.IndexOfByPredicate(
 				[&InTakeUploadID](const TSharedRef<FTakeUploadTask>& InTask)
 				{
 					return InTask->GetTaskID() == InTakeUploadID;
@@ -208,8 +215,8 @@ bool FUnrealEndpoint::FImpl::AddTakeUploadTask(TUniquePtr<FTakeUploadTask> InTak
 			if (TaskIndex != INDEX_NONE)
 			{
 				// Grab a strong reference before releasing the lock
-				TSharedRef<FTakeUploadTask> ThisTask = TakeUploadTasks[TaskIndex];
-				TakeUploadTasks.RemoveAt(TaskIndex);
+				TSharedRef<FTakeUploadTask> ThisTask = This->TakeUploadTasks[TaskIndex];
+				This->TakeUploadTasks.RemoveAt(TaskIndex);
 				ScopeLock.Unlock();
 
 				ThisTask->Complete().ExecuteIfBound(InMessage, InCode);
@@ -230,8 +237,6 @@ bool FUnrealEndpoint::FImpl::AddTakeUploadTask(TUniquePtr<FTakeUploadTask> InTak
 
 	if (!Uploader)
 	{
-		check(false);
-
 		const FString TakeInfoString = FString::Printf(
 			TEXT("Slate=%s, Take=%d, UniqueID=%s)"),
 			*SharedTask->GetTakeMetadata().Slate,
diff --git a/Engine/Plugins/VirtualProduction/CaptureManager/CaptureManagerApp/Source/CaptureManagerUnrealEndpoint/Private/DiscoveryRequester.cpp b/Engine/Plugins/VirtualProduction/CaptureManager/CaptureManagerApp/Source/CaptureManagerUnrealEndpoint/Private/DiscoveryRequester.cpp
--- a/Engine/Plugins/VirtualProduction/CaptureManager/CaptureManagerApp/Source/CaptureManagerUnrealEndpoint/Private/DiscoveryRequester.cpp
+++ b/Engine/Plugins/VirtualProduction/CaptureManager/CaptureManagerApp/Source/CaptureManagerUnrealEndpoint/Private/DiscoveryRequester.cpp
@@ -111,7 +111,7 @@ struct FDiscoveryRequester::FImpl
 
 	FString HostName;
 
-	static constexpr float DiscoveryRequestPeriod = 3.0f;
+	static constexpr float DiscoveryRequestPeriod = 15.0f;
 
 	FCaptureTimerManager::FTimerHandle TickerHandle;
 	TSharedPtr<FMessageEndpoint> MessageEndpoint;
@@ -235,13 +235,13 @@ void FDiscoveryRequester::FImpl::OnTick()
 {
 	check(!IsInGameThread());
 
-	RemoveStaleClients();
-
 	FDiscoveryRequest* Message = FMessageEndpoint::MakeMessage<FDiscoveryRequest>();
 	Message->HostName = HostName;
 
 	check(MessageEndpoint);
 	MessageEndpoint->Publish(Message);
+
+	RemoveStaleClients();
 }
 
 void FDiscoveryRequester::FImpl::HandleDiscoveryResponse(
@@ -251,6 +251,8 @@ void FDiscoveryRequester::FImpl::HandleDiscoveryResponse(
 {
 	check(!IsInGameThread());
 
+	UE_LOG(LogLiveLinkHubDiscovery, Verbose, TEXT("Getting discovery response from client: %s, (%s)"), *InContext->GetSender().ToString(), *InResponse.HostName);
+
 	FScopeLock Lock(&Mutex);
 
 	FDiscoveredClient* KnownClient =
diff --git a/Engine/Plugins/VirtualProduction/CaptureManager/CaptureManagerApp/Source/CaptureManagerUnrealEndpoint/Private/DiscoveryRequester.h b/Engine/Plugins/VirtualProduction/CaptureManager/CaptureManagerApp/Source/CaptureManagerUnrealEndpoint/Private/DiscoveryRequester.h
--- a/Engine/Plugins/VirtualProduction/CaptureManager/CaptureManagerApp/Source/CaptureManagerUnrealEndpoint/Private/DiscoveryRequester.h
+++ b/Engine/Plugins/VirtualProduction/CaptureManager/CaptureManagerApp/Source/CaptureManagerUnrealEndpoint/Private/DiscoveryRequester.h
@@ -36,7 +36,7 @@ public:
 
 private:
 
-	static constexpr double InactiveTimeout = 5.0f;
+	static constexpr double InactiveTimeout = 20.0f;
 
 	FGuid ClientID;
 	FString HostName;
diff --git a/Engine/Plugins/VirtualProduction/CaptureManager/CaptureManagerApp/Source/CaptureManagerUnrealEndpoint/Private/LiveLinkHubExportClient.cpp b/Engine/Plugins/VirtualProduction/CaptureManager/CaptureManagerApp/Source/CaptureManagerUnrealEndpoint/Private/LiveLinkHubExportClient.cpp
--- a/Engine/Plugins/VirtualProduction/CaptureManager/CaptureManagerApp/Source/CaptureManagerUnrealEndpoint/Private/LiveLinkHubExportClient.cpp
+++ b/Engine/Plugins/VirtualProduction/CaptureManager/CaptureManagerApp/Source/CaptureManagerUnrealEndpoint/Private/LiveLinkHubExportClient.cpp
@@ -245,27 +245,33 @@ void FLiveLinkHubExportClient::OnUploadTake(TUniquePtr<FTransferContext> InTrans
 		return;
 	}
 
+	TArray<FString> FilesToSend;
+
 	FUploadVoidResult FileUploadResult = MakeValue();
-	FileManager.IterateDirectoryRecursively(*Context.TakeStorage, [this, &FileUploadResult, &Context, &Token](const TCHAR* PathName, bool bIsDirectory)
+	FileManager.IterateDirectoryRecursively(*Context.TakeStorage, [this, &Token, &FilesToSend](const TCHAR* PathName, bool bIsDirectory)
 	{
-		if (bIsDirectory)
+		if (!bIsDirectory)
 		{
-			return !(Token.IsStopRequested());
+			FilesToSend.Add(PathName);
 		}
 
+		// Returns false to exit the loop if user canceled or there is an error while sending a file
+		return !(Token.IsStopRequested());
+	});
+
+	for (const FString& FilePath : FilesToSend)
+	{
 		if (Token.IsStopRequested())
 		{
 			FileUploadResult = MakeError(LOCTEXT("UploadTcpClient_Aborted", "Aborted by the user"));
+			break;
 		}
 
-		FString FileName(PathName);
+		FString FileName(FilePath);
 		FileName = FileName.Mid(Context.TakeStorage.Len() + 1); // Removing trailing slash
 
-		FileUploadResult = SendFile(FileName, PathName, Token);
-
-		// Returns false to exit the loop if user canceled or there is an error while sending a file
-		return !(Token.IsStopRequested() || FileUploadResult.HasError());
-	});
+		FileUploadResult = SendFile(FileName, FilePath, Token);
+	}
 
 	OnDataUploaded.ExecuteIfBound(TakeUploadParams.TakeUploadId, MoveTemp(FileUploadResult));
 
diff --git a/Engine/Plugins/VirtualProduction/CaptureManager/CaptureManagerApp/Source/CaptureManagerUnrealEndpoint/Public/CaptureManagerUnrealEndpoint.h b/Engine/Plugins/VirtualProduction/CaptureManager/CaptureManagerApp/Source/CaptureManagerUnrealEndpoint/Public/CaptureManagerUnrealEndpoint.h
--- a/Engine/Plugins/VirtualProduction/CaptureManager/CaptureManagerApp/Source/CaptureManagerUnrealEndpoint/Public/CaptureManagerUnrealEndpoint.h
+++ b/Engine/Plugins/VirtualProduction/CaptureManager/CaptureManagerApp/Source/CaptureManagerUnrealEndpoint/Public/CaptureManagerUnrealEndpoint.h
@@ -116,7 +116,7 @@ public:
 
 private:
 	struct FImpl;
-	TPimplPtr<FImpl> Impl;
+	TSharedPtr<FImpl> Impl;
 };
 
 }
\ No newline at end of file
diff --git a/Engine/Plugins/VirtualProduction/CaptureManager/CaptureManagerApp/Source/IngestLiveLinkDevice/Private/BaseIngestLiveLinkDevice.cpp b/Engine/Plugins/VirtualProduction/CaptureManager/CaptureManagerApp/Source/IngestLiveLinkDevice/Private/BaseIngestLiveLinkDevice.cpp
--- a/Engine/Plugins/VirtualProduction/CaptureManager/CaptureManagerApp/Source/IngestLiveLinkDevice/Private/BaseIngestLiveLinkDevice.cpp
+++ b/Engine/Plugins/VirtualProduction/CaptureManager/CaptureManagerApp/Source/IngestLiveLinkDevice/Private/BaseIngestLiveLinkDevice.cpp
@@ -233,15 +233,14 @@ void UBaseIngestLiveLinkDevice::FImpl::IngestTake(const UIngestCapability_Proces
 				Impl->TakeUploads.Remove(ProcessHandle->GetTakeId());
 				TakeUploadsRemoveLock.Unlock();
 
-				if (InStatusCode == 0)
-				{
-					Impl->Owner->ExecuteProcessFinishedReporter(ProcessHandle, MakeValue<>());
-				}
-				else
+				if (InStatusCode != 0 || !InMessage.IsEmpty())
 				{
 					TValueOrError<void, FIngestCapability_Error> Err = MakeError(FIngestCapability_Error::UnrealEndpointUploadError, InMessage);
 					Impl->Owner->ExecuteProcessFinishedReporter(ProcessHandle, MoveTemp(Err));
+					return;
 				}
+
+				Impl->Owner->ExecuteProcessFinishedReporter(ProcessHandle, MakeValue<>());
 			}
 		);
 
diff --git a/Engine/Plugins/VirtualProduction/CaptureManager/CaptureManagerCore/Source/LiveLinkHubCaptureMessaging/Private/Features/ConnectAcceptor.cpp b/Engine/Plugins/VirtualProduction/CaptureManager/CaptureManagerCore/Source/LiveLinkHubCaptureMessaging/Private/Features/ConnectAcceptor.cpp
--- a/Engine/Plugins/VirtualProduction/CaptureManager/CaptureManagerCore/Source/LiveLinkHubCaptureMessaging/Private/Features/ConnectAcceptor.cpp
+++ b/Engine/Plugins/VirtualProduction/CaptureManager/CaptureManagerCore/Source/LiveLinkHubCaptureMessaging/Private/Features/ConnectAcceptor.cpp
@@ -76,6 +76,8 @@ void FConnectAcceptor::ConnectRequestHandler(const FConnectRequest& InRequest, c
 	}
 	else
 	{
+		LastPingRequest = FPlatformTime::Seconds();
+
 		// Already connected isn't considered an error
 		Response = FMessageEndpoint::MakeMessage<FConnectResponse>();
 	}
diff --git a/Engine/Plugins/VirtualProduction/CaptureManager/CaptureManagerCore/Source/LiveLinkHubCaptureMessaging/Private/Features/ConnectStarter.cpp b/Engine/Plugins/VirtualProduction/CaptureManager/CaptureManagerCore/Source/LiveLinkHubCaptureMessaging/Private/Features/ConnectStarter.cpp
--- a/Engine/Plugins/VirtualProduction/CaptureManager/CaptureManagerCore/Source/LiveLinkHubCaptureMessaging/Private/Features/ConnectStarter.cpp
+++ b/Engine/Plugins/VirtualProduction/CaptureManager/CaptureManagerCore/Source/LiveLinkHubCaptureMessaging/Private/Features/ConnectStarter.cpp
@@ -18,7 +18,7 @@ static TSharedRef<FCaptureTimerManager> GetTimerManager()
 
 }
 
-const float FConnectStarter::KeepAliveInterval = 5.0f; // seconds
+const float FConnectStarter::KeepAliveInterval = 7.5f; // seconds KeepAliveTimeout * 2 + KeepAliveTimeout / 2
 const float FConnectStarter::KeepAliveTimeout = 3.0f; // seconds
 
 FConnectStarter::FConnectStarter()
@@ -49,7 +49,6 @@ void FConnectStarter::Connect(FConnectHandler InConnectHandler)
 	}
 
 	FConnectRequest* Request = FMessageEndpoint::MakeMessage<FConnectRequest>();
-
 	Request->Guid = FGuid::NewGuid();
 
 	AddContext(Request->Guid, MoveTemp(InConnectHandler));
@@ -86,14 +85,16 @@ void FConnectStarter::Initialize(FMessageEndpointBuilder& InBuilder)
 		.Handling<FPongMessage>(this, &FConnectStarter::HandlePingResponse);
 }
 
-void FConnectStarter::SendPingRequest(FPingCallback InResponseCb)
+FGuid FConnectStarter::SendPingRequest(FPingCallback InResponseCb)
 {
 	if (!bConnected)
 	{
-		return;
+		return FGuid();
 	}
 
 	FPingMessage* Request = FMessageEndpoint::MakeMessage<FPingMessage>();
+	Request->Guid = FGuid::NewGuid();
+
 	FGuid RequestGuid = Request->Guid;
 
 	{
@@ -102,6 +103,8 @@ void FConnectStarter::SendPingRequest(FPingCallback InResponseCb)
 	}
 
 	Endpoint->Send(Request, Address);
+
+	return RequestGuid;
 }
 
 void FConnectStarter::HandleConnectResponse(const FConnectResponse& InResponse, const TSharedRef<IMessageContext, ESPMode::ThreadSafe>& InContext)
@@ -149,19 +152,20 @@ void FConnectStarter::OnKeepAliveInterval()
 	{
 		FSharedEventRef PongEvent;
 
-		FPingMessage* Request = FMessageEndpoint::MakeMessage<FPingMessage>();
-		FGuid RequestGuid = Request->Guid;
-
-		This->SendPingRequest(FPingCallback::CreateLambda([PongEvent](const FPongMessage& InResponse) mutable
+		FGuid RequestGuid = This->SendPingRequest(FPingCallback::CreateLambda([PongEvent](const FPongMessage& InResponse) mutable
 		{
 			PongEvent->Trigger();
 		}));
 
-		This->Endpoint->Send(Request, This->Address);
+		if (!RequestGuid.IsValid())
+		{
+			// We already disconnected
+			return;
+		}
 
-		bool Result = PongEvent->Wait(FTimespan::FromSeconds(KeepAliveTimeout));
+		bool bResult = PongEvent->Wait(FTimespan::FromSeconds(KeepAliveTimeout));
 
-		if (!Result)
+		if (!bResult)
 		{
 			This->bConnected = false;
 
@@ -175,7 +179,7 @@ void FConnectStarter::OnKeepAliveInterval()
 
 		This->TimerManager->RemoveTimer(This->KeepAlive);
 
-		if (Result)
+		if (bResult)
 		{
 			This->KeepAlive = 
 				This->TimerManager->AddTimer(
diff --git a/Engine/Plugins/VirtualProduction/CaptureManager/CaptureManagerCore/Source/LiveLinkHubCaptureMessaging/Public/Features/ConnectAcceptor.h b/Engine/Plugins/VirtualProduction/CaptureManager/CaptureManagerCore/Source/LiveLinkHubCaptureMessaging/Public/Features/ConnectAcceptor.h
--- a/Engine/Plugins/VirtualProduction/CaptureManager/CaptureManagerCore/Source/LiveLinkHubCaptureMessaging/Public/Features/ConnectAcceptor.h
+++ b/Engine/Plugins/VirtualProduction/CaptureManager/CaptureManagerCore/Source/LiveLinkHubCaptureMessaging/Public/Features/ConnectAcceptor.h
@@ -38,7 +38,7 @@ private:
 	void CheckConnectionActivity();
 
 	double LastPingRequest = 0.0f;
-	double InactivityTimeout = 15.0f; // seconds
+	double InactivityTimeout = 20.0f; // seconds
 
 	FConnectAccepted ConnectHandler;
 	FConnectionLostHandler ConnectionLostHandler;
diff --git a/Engine/Plugins/VirtualProduction/CaptureManager/CaptureManagerCore/Source/LiveLinkHubCaptureMessaging/Public/Features/ConnectStarter.h b/Engine/Plugins/VirtualProduction/CaptureManager/CaptureManagerCore/Source/LiveLinkHubCaptureMessaging/Public/Features/ConnectStarter.h
--- a/Engine/Plugins/VirtualProduction/CaptureManager/CaptureManagerCore/Source/LiveLinkHubCaptureMessaging/Public/Features/ConnectStarter.h
+++ b/Engine/Plugins/VirtualProduction/CaptureManager/CaptureManagerCore/Source/LiveLinkHubCaptureMessaging/Public/Features/ConnectStarter.h
@@ -39,7 +39,7 @@ private:
 	static const float KeepAliveTimeout;
 
 	DECLARE_DELEGATE_OneParam(FPingCallback, const FPongMessage& InResponse);
-	void SendPingRequest(FPingCallback InCallback);
+	FGuid SendPingRequest(FPingCallback InCallback);
 
 	void HandleConnectResponse(const FConnectResponse& InResponse, const TSharedRef<IMessageContext, ESPMode::ThreadSafe>& InContext);
 	void HandlePingResponse(const FPongMessage& InResponse, const TSharedRef<IMessageContext, ESPMode::ThreadSafe>& InContext);
diff --git a/Engine/Plugins/VirtualProduction/CaptureManager/CaptureManagerEditor/Source/LiveLinkHubDiscoveryEditor/Private/DiscoveryResponder.cpp b/Engine/Plugins/VirtualProduction/CaptureManager/CaptureManagerEditor/Source/LiveLinkHubDiscoveryEditor/Private/DiscoveryResponder.cpp
--- a/Engine/Plugins/VirtualProduction/CaptureManager/CaptureManagerEditor/Source/LiveLinkHubDiscoveryEditor/Private/DiscoveryResponder.cpp
+++ b/Engine/Plugins/VirtualProduction/CaptureManager/CaptureManagerEditor/Source/LiveLinkHubDiscoveryEditor/Private/DiscoveryResponder.cpp
@@ -96,15 +96,21 @@ void FDiscoveryResponder::HandleDiscoveryRequest(
 	using FServerInfo = FLiveLinkHubExportServer::FServerInfo;
 	using EServerError = FLiveLinkHubExportServer::EServerError;
 
-	TValueOrError<FServerInfo, EServerError> Info = GetExportServerInfo();
-	if (Info.HasError())
+	if (!CachedServerInfo.IsSet())
 	{
-		return;
+		TValueOrError<FServerInfo, EServerError> Info = GetExportServerInfo();
+		if (Info.HasError())
+		{
+			CachedServerInfo.Reset();
+			return;
+		}
+
+		CachedServerInfo = Info.StealValue();
 	}
 
 	FDiscoveryResponse* Response = FMessageEndpoint::MakeMessage<FDiscoveryResponse>();
-	Response->ExportPort = Info.GetValue().Port;
-	Response->IPAddress = Info.GetValue().IPAddress;
+	Response->ExportPort = CachedServerInfo.GetValue().Port;
+	Response->IPAddress = CachedServerInfo.GetValue().IPAddress;
 	Response->HostName = HostName;
 
 	Manager->SendDiscoveryResponse(Response, InContext->GetSender());
@@ -115,6 +121,15 @@ void FDiscoveryResponder::HandleDiscoveryRequest(
 void FDiscoveryResponder::StartDiscoveryResponder()
 {
 	bStarted = true;
+
+	TValueOrError<FLiveLinkHubExportServer::FServerInfo, FLiveLinkHubExportServer::EServerError> Info = GetExportServerInfo();
+	if (Info.HasError())
+	{
+		CachedServerInfo.Reset();
+		return;
+	}
+
+	CachedServerInfo = Info.StealValue();
 }
 
 } // namespace UE::CaptureManager
diff --git a/Engine/Plugins/VirtualProduction/CaptureManager/CaptureManagerEditor/Source/LiveLinkHubDiscoveryEditor/Private/DiscoveryResponder.h b/Engine/Plugins/VirtualProduction/CaptureManager/CaptureManagerEditor/Source/LiveLinkHubDiscoveryEditor/Private/DiscoveryResponder.h
--- a/Engine/Plugins/VirtualProduction/CaptureManager/CaptureManagerEditor/Source/LiveLinkHubDiscoveryEditor/Private/DiscoveryResponder.h
+++ b/Engine/Plugins/VirtualProduction/CaptureManager/CaptureManagerEditor/Source/LiveLinkHubDiscoveryEditor/Private/DiscoveryResponder.h
@@ -36,6 +36,7 @@ private:
 	TSharedRef<FLiveLinkHubWorkerManager> Manager;
 
 	FString HostName;
+	TOptional<FLiveLinkHubExportServer::FServerInfo> CachedServerInfo;
 
 	TSharedPtr<FMessageEndpoint> MessageEndpoint;
 
diff --git a/Engine/Plugins/VirtualProduction/CaptureManager/CaptureManagerEditor/Source/LiveLinkHubWorkerManager/Private/LiveLinkHubWorker.cpp b/Engine/Plugins/VirtualProduction/CaptureManager/CaptureManagerEditor/Source/LiveLinkHubWorkerManager/Private/LiveLinkHubWorker.cpp
--- a/Engine/Plugins/VirtualProduction/CaptureManager/CaptureManagerEditor/Source/LiveLinkHubWorkerManager/Private/LiveLinkHubWorker.cpp
+++ b/Engine/Plugins/VirtualProduction/CaptureManager/CaptureManagerEditor/Source/LiveLinkHubWorkerManager/Private/LiveLinkHubWorker.cpp
@@ -63,14 +63,23 @@ static void PublishDoneEvent(TSharedPtr<FLiveLinkHubImportWorker::FEditorMesseng
 	InMessenger->SendUploadDoneMessage(InCaptureSourceId, InTakeUploadId, MoveTemp(InMessage), InCode);
 }
 
-FLiveLinkHubImportWorker::FLiveLinkHubImportWorker(TWeakPtr<FEditorMessenger> InEditorMessenger)
-	: Messenger(MoveTemp(InEditorMessenger))
+TSharedPtr<FLiveLinkHubImportWorker> FLiveLinkHubImportWorker::Create(TWeakPtr<FEditorMessenger> InEditorMessenger)
 {
-	TSharedPtr<FEditorMessenger> SharedMessenger = Messenger.Pin();
+	TSharedPtr<FLiveLinkHubImportWorker> SharedWorker = 
+		MakeShared<FLiveLinkHubImportWorker>(MoveTemp(InEditorMessenger), FPrivateToken());
+
+	TSharedPtr<FEditorMessenger> SharedMessenger = SharedWorker->Messenger.Pin();
 
 	FLiveLinkHubExportServerModule& Module = FModuleManager::LoadModuleChecked<FLiveLinkHubExportServerModule>("LiveLinkHubExportServer");
 	Module.RegisterExportServerHandler(SharedMessenger->GetAddress().ToString(),
-									   FLiveLinkHubExportServer::FFileDataHandler::CreateRaw(this, &FLiveLinkHubImportWorker::HandleTakeDownload));
+									   FLiveLinkHubExportServer::FFileDataHandler::CreateSP(SharedWorker.ToSharedRef(), &FLiveLinkHubImportWorker::HandleTakeDownload));
+
+	return SharedWorker;
+}
+
+FLiveLinkHubImportWorker::FLiveLinkHubImportWorker(TWeakPtr<FEditorMessenger> InEditorMessenger, FPrivateToken)
+	: Messenger(MoveTemp(InEditorMessenger))
+{
 }
 
 FLiveLinkHubImportWorker::~FLiveLinkHubImportWorker()
@@ -222,7 +231,13 @@ FUploadVoidResult FLiveLinkHubImportWorker::HandleFileDownload(const FString& In
 	IFileManager& FileManager = IFileManager::Get();
 
 	TUniquePtr<FArchive> Writer(FileManager.CreateFileWriter(*FilePathToBeSaved, EFileWrite::FILEWRITE_None));
-	check(Writer);
+	if (!Writer)
+	{
+		FText Message =
+			FText::Format(LOCTEXT("HandleFileDownload_FailedToCreateFile", "Failed to create the file: {0}"),
+						  FText::FromString(FilePathToBeSaved));
+		return MakeError(MoveTemp(Message));
+	}
 
 	FMD5 MD5Generator;
 
@@ -302,8 +317,6 @@ void FLiveLinkHubImportWorker::RemoveContext(const FGuid& InUploadId)
 	TakeFilesContext.Remove(InUploadId);
 }
 
-#define CHECK_AND_PUBLISH_EVENT(Result, Message) if (!Result) { PublishDoneEvent(TakeName, Message, -30); }
-
 void FLiveLinkHubImportWorker::SpawnIngestTask(TSharedPtr<FEditorMessenger> InMessenger,
 	const FString& InDataStorage,
 	const FGuid& InCaptureSourceId,
@@ -313,7 +326,7 @@ void FLiveLinkHubImportWorker::SpawnIngestTask(TSharedPtr<FEditorMessenger> InMe
 {
 	using namespace UE::CaptureManager;
 
-	AsyncTask(ENamedThreads::AnyThread, [this,
+	AsyncTask(ENamedThreads::AnyThread, [
 		SharedMessenger = MoveTemp(InMessenger),
 		DataStorage = InDataStorage,
 		CaptureSourceId = InCaptureSourceId,
@@ -329,10 +342,10 @@ void FLiveLinkHubImportWorker::SpawnIngestTask(TSharedPtr<FEditorMessenger> InMe
 			{
 				FIngestProcessResult IngestProcessResult = ProcessResult.GetValue();
 
-				ExecuteOnGameThread(TEXT("IngestAssetCreation"), [this, SharedMessenger, IngestProcessResult = MoveTemp(IngestProcessResult), CaptureSourceId, CaptureSourceName, TakeUploadId]()
+				ExecuteOnGameThread(TEXT("IngestAssetCreation"), [SharedMessenger, IngestProcessResult = MoveTemp(IngestProcessResult), CaptureSourceId, CaptureSourceName, TakeUploadId]()
 					{
 						FIngestAssetCreator::FPerTakeCallback Callback = FIngestAssetCreator::FPerTakeCallback(
-							FIngestAssetCreator::FPerTakeCallback::Type::CreateLambda([this, SharedMessenger, TakeUploadId, CaptureSourceId](TPair<int32, FIngestAssetCreator::FAssetCreationResult> InResult)
+							FIngestAssetCreator::FPerTakeCallback::Type::CreateLambda([SharedMessenger, TakeUploadId, CaptureSourceId](TPair<int32, FIngestAssetCreator::FAssetCreationResult> InResult)
 								{
 									if (InResult.Value.HasError())
 									{
diff --git a/Engine/Plugins/VirtualProduction/CaptureManager/CaptureManagerEditor/Source/LiveLinkHubWorkerManager/Private/LiveLinkHubWorkerManager.cpp b/Engine/Plugins/VirtualProduction/CaptureManager/CaptureManagerEditor/Source/LiveLinkHubWorkerManager/Private/LiveLinkHubWorkerManager.cpp
--- a/Engine/Plugins/VirtualProduction/CaptureManager/CaptureManagerEditor/Source/LiveLinkHubWorkerManager/Private/LiveLinkHubWorkerManager.cpp
+++ b/Engine/Plugins/VirtualProduction/CaptureManager/CaptureManagerEditor/Source/LiveLinkHubWorkerManager/Private/LiveLinkHubWorkerManager.cpp
@@ -41,15 +41,15 @@ void FLiveLinkHubWorkerManager::SendDiscoveryResponse(FDiscoveryResponse* Respon
 	Messenger->SendDiscoveryResponse(Response, Receiver);
 }
 
-UE::CaptureManager::Private::FLiveLinkHubImportWorker* FLiveLinkHubWorkerManager::AddWorker(FMessageAddress InServer)
+TWeakPtr<UE::CaptureManager::Private::FLiveLinkHubImportWorker> FLiveLinkHubWorkerManager::AddWorker(FMessageAddress InServer)
 {
 	using namespace UE::CaptureManager::Private;
 
 	FScopeLock Lock(&Mutex);
 
-	TUniquePtr<FLiveLinkHubImportWorker>& AddedWorker = Workers.Add(InServer.ToString(), MakeUnique<FLiveLinkHubImportWorker>(Messenger));
+	TSharedPtr<FLiveLinkHubImportWorker>& AddedWorker = Workers.Add(InServer.ToString(), FLiveLinkHubImportWorker::Create(Messenger));
 
-	return AddedWorker.Get();
+	return AddedWorker;
 }
 
 void FLiveLinkHubWorkerManager::RemoveWorker(FMessageAddress InServer)
diff --git a/Engine/Plugins/VirtualProduction/CaptureManager/CaptureManagerEditor/Source/LiveLinkHubWorkerManager/Public/LiveLinkHubWorker.h b/Engine/Plugins/VirtualProduction/CaptureManager/CaptureManagerEditor/Source/LiveLinkHubWorkerManager/Public/LiveLinkHubWorker.h
--- a/Engine/Plugins/VirtualProduction/CaptureManager/CaptureManagerEditor/Source/LiveLinkHubWorkerManager/Public/LiveLinkHubWorker.h
+++ b/Engine/Plugins/VirtualProduction/CaptureManager/CaptureManagerEditor/Source/LiveLinkHubWorkerManager/Public/LiveLinkHubWorker.h
@@ -26,12 +26,18 @@ struct FCaptureDataAssetInfo;
 namespace Private
 {
 
-class FLiveLinkHubImportWorker
+class FLiveLinkHubImportWorker : public TSharedFromThis<FLiveLinkHubImportWorker>
 {
+private:
+
+	struct FPrivateToken {};
+
 public:
 	using FEditorMessenger = FMessenger<FConnectAcceptor, FUploadStateSender>;
 
-	FLiveLinkHubImportWorker(TWeakPtr<FEditorMessenger> InEditorMessenger);
+	static TSharedPtr<FLiveLinkHubImportWorker> Create(TWeakPtr<FEditorMessenger> InEditorMessenger);
+
+	explicit FLiveLinkHubImportWorker(TWeakPtr<FEditorMessenger> InEditorMessenger, FPrivateToken);
 	~FLiveLinkHubImportWorker();
 
 private:
@@ -56,18 +62,18 @@ private:
 	FTakeFileContext& AddOrIgnoreContext(const FGuid& InUploadId, uint64 InTotalLength);
 	void RemoveContext(const FGuid& InUploadId);
 
-	void SpawnIngestTask(TSharedPtr<FEditorMessenger> InMessenger,
+	static void SpawnIngestTask(TSharedPtr<FEditorMessenger> InMessenger,
 		const FString& InDataStorage,
 		const FGuid& InCaptureSourceId,
 		const FString& InCaptureSourceName,
 		const FGuid& InTakeUploadId,
 		TSharedPtr<UE::CaptureManager::FTaskProgress> InTaskProgress);
 
-	bool CreateCaptureAsset(const FString& InAssetPath,
+	static bool CreateCaptureAsset(const FString& InAssetPath,
 		const UE::CaptureManager::FCaptureDataAssetInfo& InResult,
 		const FCaptureDataTakeInfo& InTakeInfo);
 
-	void SaveCaptureCreatedAssets(const FString& InAssetPath);
+	static void SaveCaptureCreatedAssets(const FString& InAssetPath);
 
 	FCriticalSection ExportMutex;
 	TMap<FGuid, FTakeFileContext> TakeFilesContext;
diff --git a/Engine/Plugins/VirtualProduction/CaptureManager/CaptureManagerEditor/Source/LiveLinkHubWorkerManager/Public/LiveLinkHubWorkerManager.h b/Engine/Plugins/VirtualProduction/CaptureManager/CaptureManagerEditor/Source/LiveLinkHubWorkerManager/Public/LiveLinkHubWorkerManager.h
--- a/Engine/Plugins/VirtualProduction/CaptureManager/CaptureManagerEditor/Source/LiveLinkHubWorkerManager/Public/LiveLinkHubWorkerManager.h
+++ b/Engine/Plugins/VirtualProduction/CaptureManager/CaptureManagerEditor/Source/LiveLinkHubWorkerManager/Public/LiveLinkHubWorkerManager.h
@@ -18,7 +18,7 @@ public:
 
 private:
 
-	UE::CaptureManager::Private::FLiveLinkHubImportWorker* AddWorker(FMessageAddress InServer);
+	TWeakPtr<UE::CaptureManager::Private::FLiveLinkHubImportWorker> AddWorker(FMessageAddress InServer);
 
 	void RemoveWorker(FMessageAddress InServer);
 
@@ -28,5 +28,5 @@ private:
 	TSharedPtr<UE::CaptureManager::Private::FLiveLinkHubImportWorker::FEditorMessenger> Messenger;
 
 	FCriticalSection Mutex;
-	TMap<FString, TUniquePtr<UE::CaptureManager::Private::FLiveLinkHubImportWorker>> Workers;
+	TMap<FString, TSharedPtr<UE::CaptureManager::Private::FLiveLinkHubImportWorker>> Workers;
 };
\ No newline at end of file
