Capture Manager Job details not inheriting conversion folder from settings

diff --git a/Engine/Plugins/VirtualProduction/CaptureManager/CaptureManagerApp/Source/CaptureManagerEditor/Private/CaptureManagerEditorModule.cpp b/Engine/Plugins/VirtualProduction/CaptureManager/CaptureManagerApp/Source/CaptureManagerEditor/Private/CaptureManagerEditorModule.cpp
--- a/Engine/Plugins/VirtualProduction/CaptureManager/CaptureManagerApp/Source/CaptureManagerEditor/Private/CaptureManagerEditorModule.cpp
+++ b/Engine/Plugins/VirtualProduction/CaptureManager/CaptureManagerApp/Source/CaptureManagerEditor/Private/CaptureManagerEditorModule.cpp
@@ -25,7 +25,7 @@ void FCaptureManagerEditorModule::StartupModule()
 
 	FPropertyEditorModule& PropertyEditorModule = FModuleManager::GetModuleChecked<FPropertyEditorModule>("PropertyEditor");
 	PropertyEditorModule.RegisterCustomClassLayout(
-		UIngestJobSettingsBase::StaticClass()->GetFName(),
+		UIngestJobSettings::StaticClass()->GetFName(),
 		FOnGetDetailCustomizationInstance::CreateStatic(&FIngestJobSettingsCustomization::MakeInstance)
 	);
 }
@@ -38,7 +38,7 @@ void FCaptureManagerEditorModule::ShutdownModule()
 	}
 
 	FPropertyEditorModule& PropertyEditorModule = FModuleManager::GetModuleChecked<FPropertyEditorModule>("PropertyEditor");
-	PropertyEditorModule.UnregisterCustomClassLayout(UIngestJobSettingsBase::StaticClass()->GetFName());
+	PropertyEditorModule.UnregisterCustomClassLayout(UIngestJobSettings::StaticClass()->GetFName());
 }
 
 /** Instantiate an application mode so LiveLinkHub can register it and display it in its Layout Selector. */
diff --git a/Engine/Plugins/VirtualProduction/CaptureManager/CaptureManagerApp/Source/CaptureManagerEditor/Private/IngestManagement/IngestJobSettingsManager.cpp b/Engine/Plugins/VirtualProduction/CaptureManager/CaptureManagerApp/Source/CaptureManagerEditor/Private/IngestManagement/IngestJobSettingsManager.cpp
--- a/Engine/Plugins/VirtualProduction/CaptureManager/CaptureManagerApp/Source/CaptureManagerEditor/Private/IngestManagement/IngestJobSettingsManager.cpp
+++ b/Engine/Plugins/VirtualProduction/CaptureManager/CaptureManagerApp/Source/CaptureManagerEditor/Private/IngestManagement/IngestJobSettingsManager.cpp
@@ -5,73 +5,15 @@
 namespace UE::CaptureManager
 {
 
-namespace Private
+FIngestJobSettingsManager::FIngestJobSettingsManager()
 {
-static TArray<TStrongObjectPtr<const UIngestJobSettingsBase>> GetIngestJobPresets()
-{
-	TArray<TStrongObjectPtr<const UIngestJobSettingsBase>> Presets;
-	Presets.Emplace(NewObject<UDefaultJobPreset>());
-	return Presets;
-}
-} // namespace UE::CaptureManager::Private
-
-FIngestJobSettingsManager::FIngestJobSettingsManager() :
-	Presets(Private::GetIngestJobPresets())
-{
-	check(!Presets.IsEmpty());
-
-	if (!Presets.IsEmpty())
-	{
-		SelectedPreset = Presets[0].Get();
-	}
-}
-
-TArray<FText> FIngestJobSettingsManager::GetPresetDisplayNames() const
-{
-	FScopeLock Lock(&CriticalSection);
-
-	TArray<FText> Names;
-	Names.Reserve(Presets.Num());
-
-	for (const TStrongObjectPtr<const UIngestJobSettingsBase>& Preset : Presets)
-	{
-		Names.Emplace(Preset->DisplayName);
-	}
-
-	return Names;
-}
-
-bool FIngestJobSettingsManager::SelectPresetByDisplayName(const FText& DisplayName)
-{
-	FScopeLock Lock(&CriticalSection);
-
-	const TStrongObjectPtr<const UIngestJobSettingsBase>* Preset = Presets.FindByPredicate(
-		[&DisplayName](const TStrongObjectPtr<const UIngestJobSettingsBase>& InPreset)
-		{
-			return InPreset && InPreset->DisplayName.CompareTo(DisplayName) == 0;
-		}
-	);
-
-	if (Preset)
-	{
-		SelectedPreset = Preset->Get();
-		return true;
-	}
-
-	return false;
-}
-
-TStrongObjectPtr<const UIngestJobSettingsBase> FIngestJobSettingsManager::GetSelectedPreset() const
-{
-	FScopeLock Lock(&CriticalSection);
-	return TStrongObjectPtr<const UIngestJobSettingsBase>(SelectedPreset);
 }
 
 void FIngestJobSettingsManager::ApplyJobSpecificSettings(const FGuid& InJobGuid, const FIngestJob::FSettings& InSettings)
 {
 	FScopeLock Lock(&CriticalSection);
 
-	TStrongObjectPtr<UIngestJobSettingsBase>& Object = Settings.Emplace_GetRef(DuplicateObject(SelectedPreset, nullptr));
+	TStrongObjectPtr<UIngestJobSettings>& Object = Settings.Emplace_GetRef(NewObject<UIngestJobSettings>());
 	Object->JobGuid = InJobGuid;
 	Object->AudioFileNamePrefix = InSettings.AudioSettings.FileNamePrefix;
 	Object->AudioFormat = InSettings.AudioSettings.Format;
@@ -89,7 +31,7 @@ int32 FIngestJobSettingsManager::RemoveSettings(const TArray<FGuid>& JobGuids)
 	FScopeLock Lock(&CriticalSection);
 
 	const int32 NumRemoved = Settings.RemoveAll(
-		[&JobGuids](const TStrongObjectPtr<UIngestJobSettingsBase>& Entry)
+		[&JobGuids](const TStrongObjectPtr<UIngestJobSettings>& Entry)
 		{
 			return JobGuids.Contains(Entry->JobGuid);
 		}
@@ -98,14 +40,14 @@ int32 FIngestJobSettingsManager::RemoveSettings(const TArray<FGuid>& JobGuids)
 	return NumRemoved;
 }
 
-TArray<TWeakObjectPtr<UIngestJobSettingsBase>> FIngestJobSettingsManager::GetSettings(const TArray<FGuid>& InJobGuids) const
+TArray<TWeakObjectPtr<UIngestJobSettings>> FIngestJobSettingsManager::GetSettings(const TArray<FGuid>& InJobGuids) const
 {
 	FScopeLock Lock(&CriticalSection);
 
-	TArray<TWeakObjectPtr<UIngestJobSettingsBase>> JobSettings;
+	TArray<TWeakObjectPtr<UIngestJobSettings>> JobSettings;
 	JobSettings.Reserve(InJobGuids.Num());
 
-	for (const TStrongObjectPtr<UIngestJobSettingsBase>& Entry : Settings)
+	for (const TStrongObjectPtr<UIngestJobSettings>& Entry : Settings)
 	{
 		if (InJobGuids.Contains(Entry->JobGuid))
 		{
diff --git a/Engine/Plugins/VirtualProduction/CaptureManager/CaptureManagerApp/Source/CaptureManagerEditor/Private/IngestManagement/IngestJobSettingsManager.h b/Engine/Plugins/VirtualProduction/CaptureManager/CaptureManagerApp/Source/CaptureManagerEditor/Private/IngestManagement/IngestJobSettingsManager.h
--- a/Engine/Plugins/VirtualProduction/CaptureManager/CaptureManagerApp/Source/CaptureManagerEditor/Private/IngestManagement/IngestJobSettingsManager.h
+++ b/Engine/Plugins/VirtualProduction/CaptureManager/CaptureManagerApp/Source/CaptureManagerEditor/Private/IngestManagement/IngestJobSettingsManager.h
@@ -14,19 +14,13 @@ class FIngestJobSettingsManager
 public:
 	FIngestJobSettingsManager();
 
-	TArray<FText> GetPresetDisplayNames() const;
-	bool SelectPresetByDisplayName(const FText& DisplayName);
-
-	TStrongObjectPtr<const UIngestJobSettingsBase> GetSelectedPreset() const;
 	void ApplyJobSpecificSettings(const FGuid& InJobGuid, const FIngestJob::FSettings& InSettings);
-	TArray<TWeakObjectPtr<UIngestJobSettingsBase>> GetSettings(const TArray<FGuid>& InJobGuids) const;
+	TArray<TWeakObjectPtr<UIngestJobSettings>> GetSettings(const TArray<FGuid>& InJobGuids) const;
 
 	int32 RemoveSettings(const TArray<FGuid>& InJobGuids);
 
 private:
 	mutable FCriticalSection CriticalSection;
-	const UIngestJobSettingsBase* SelectedPreset = nullptr;
-	const TArray<TStrongObjectPtr<const UIngestJobSettingsBase>> Presets;
-	TArray<TStrongObjectPtr<UIngestJobSettingsBase>> Settings;
+	TArray<TStrongObjectPtr<UIngestJobSettings>> Settings;
 };
 }
\ No newline at end of file
diff --git a/Engine/Plugins/VirtualProduction/CaptureManager/CaptureManagerApp/Source/CaptureManagerEditor/Private/IngestManagement/UIngestJobSettings.cpp b/Engine/Plugins/VirtualProduction/CaptureManager/CaptureManagerApp/Source/CaptureManagerEditor/Private/IngestManagement/UIngestJobSettings.cpp
--- a/Engine/Plugins/VirtualProduction/CaptureManager/CaptureManagerApp/Source/CaptureManagerEditor/Private/IngestManagement/UIngestJobSettings.cpp
+++ b/Engine/Plugins/VirtualProduction/CaptureManager/CaptureManagerApp/Source/CaptureManagerEditor/Private/IngestManagement/UIngestJobSettings.cpp
@@ -8,12 +8,7 @@
 
 #define LOCTEXT_NAMESPACE "UIngestJobSettings"
 
-UIngestJobSettingsBase::UIngestJobSettingsBase(const FObjectInitializer& ObjectInitializer)
-	: Super(ObjectInitializer)
-{
-}
-
-UDefaultJobPreset::UDefaultJobPreset(const FObjectInitializer& ObjectInitializer)
+UIngestJobSettings::UIngestJobSettings(const FObjectInitializer& ObjectInitializer)
 	: Super(ObjectInitializer)
 {
 	DisplayName = LOCTEXT("DefaultText", "Default");
@@ -34,7 +29,7 @@ UDefaultJobPreset::UDefaultJobPreset(const FObjectInitializer& ObjectInitializer
 	check(!UploadHostName.IsEmpty());
 }
 
-void UDefaultJobPreset::PostEditChangeProperty(struct FPropertyChangedEvent& InPropertyChangedEvent)
+void UIngestJobSettings::PostEditChangeProperty(struct FPropertyChangedEvent& InPropertyChangedEvent)
 {
 	Super::PostEditChangeProperty(InPropertyChangedEvent);
 
diff --git a/Engine/Plugins/VirtualProduction/CaptureManager/CaptureManagerApp/Source/CaptureManagerEditor/Private/IngestManagement/UIngestJobSettings.h b/Engine/Plugins/VirtualProduction/CaptureManager/CaptureManagerApp/Source/CaptureManagerEditor/Private/IngestManagement/UIngestJobSettings.h
--- a/Engine/Plugins/VirtualProduction/CaptureManager/CaptureManagerApp/Source/CaptureManagerEditor/Private/IngestManagement/UIngestJobSettings.h
+++ b/Engine/Plugins/VirtualProduction/CaptureManager/CaptureManagerApp/Source/CaptureManagerEditor/Private/IngestManagement/UIngestJobSettings.h
@@ -38,8 +38,8 @@ enum class EAudioFormat : uint8
 	WAV = 0				UMETA(DisplayName = "wav")
 };
 
-UCLASS(Abstract)
-class UIngestJobSettingsBase : public UObject
+UCLASS()
+class UIngestJobSettings : public UObject
 {
 	GENERATED_BODY()
 
@@ -77,18 +77,9 @@ public:
 	UPROPERTY(Category = "Upload", EditAnywhere, Meta = (DisplayName = "Host Name", NoResetToDefault))
 	FString UploadHostName;
 
-	UIngestJobSettingsBase(const FObjectInitializer& ObjectInitializer);
-};
-
-UCLASS()
-class UDefaultJobPreset : public UIngestJobSettingsBase
-{
-	GENERATED_BODY()
-
-public:
-	UDefaultJobPreset(const FObjectInitializer& ObjectInitializer);
+	UIngestJobSettings(const FObjectInitializer& ObjectInitializer);
 
 private:
 
-	virtual void PostEditChangeProperty(struct FPropertyChangedEvent& InPropertyChangedEvent) override;
+	virtual void PostEditChangeProperty(FPropertyChangedEvent& PropertyChangedEvent) override;
 };
diff --git a/Engine/Plugins/VirtualProduction/CaptureManager/CaptureManagerApp/Source/CaptureManagerEditor/Private/UI/CaptureManagerPanelController.cpp b/Engine/Plugins/VirtualProduction/CaptureManager/CaptureManagerApp/Source/CaptureManagerEditor/Private/UI/CaptureManagerPanelController.cpp
--- a/Engine/Plugins/VirtualProduction/CaptureManager/CaptureManagerApp/Source/CaptureManagerEditor/Private/UI/CaptureManagerPanelController.cpp
+++ b/Engine/Plugins/VirtualProduction/CaptureManager/CaptureManagerApp/Source/CaptureManagerEditor/Private/UI/CaptureManagerPanelController.cpp
@@ -319,7 +319,7 @@ void FCaptureManagerPanelController::OnFinishedEditingJobProperties(const FPrope
 
 	for (int32 Index = 0; Index < PropertyChangedEvent.GetNumObjectsBeingEdited(); Index++)
 	{
-		if (const UIngestJobSettingsBase* Settings = Cast<const UIngestJobSettingsBase>(PropertyChangedEvent.GetObjectBeingEdited(Index)))
+		if (const UIngestJobSettings* Settings = Cast<const UIngestJobSettings>(PropertyChangedEvent.GetObjectBeingEdited(Index)))
 		{
 			// It shouldn't be possible to set an empty upload host name
 			check(!Settings->UploadHostName.IsEmpty());
@@ -358,12 +358,12 @@ void FCaptureManagerPanelController::OnIngestJobSelectionChanged(const TArray<FG
 
 	check(IngestJobDetailsView);
 
-	TArray<TWeakObjectPtr<UIngestJobSettingsBase>> SettingsForJobs = IngestJobSettingsManager->GetSettings(InJobGuids);
+	TArray<TWeakObjectPtr<UIngestJobSettings>> SettingsForJobs = IngestJobSettingsManager->GetSettings(InJobGuids);
 
 	TArray<TWeakObjectPtr<UObject>> ObjectsForDetailsView;
 	ObjectsForDetailsView.Reserve(SettingsForJobs.Num());
 
-	for (TWeakObjectPtr<UIngestJobSettingsBase> SettingsForJob : SettingsForJobs)
+	for (TWeakObjectPtr<UIngestJobSettings> SettingsForJob : SettingsForJobs)
 	{
 		ObjectsForDetailsView.Emplace(SettingsForJob);
 	}
diff --git a/Engine/Plugins/VirtualProduction/CaptureManager/CaptureManagerApp/Source/CaptureManagerEditor/Private/UI/IngestJobSettingsCustomization.cpp b/Engine/Plugins/VirtualProduction/CaptureManager/CaptureManagerApp/Source/CaptureManagerEditor/Private/UI/IngestJobSettingsCustomization.cpp
--- a/Engine/Plugins/VirtualProduction/CaptureManager/CaptureManagerApp/Source/CaptureManagerEditor/Private/UI/IngestJobSettingsCustomization.cpp
+++ b/Engine/Plugins/VirtualProduction/CaptureManager/CaptureManagerApp/Source/CaptureManagerEditor/Private/UI/IngestJobSettingsCustomization.cpp
@@ -17,7 +17,7 @@ TSharedRef<IDetailCustomization> FIngestJobSettingsCustomization::MakeInstance()
 
 void FIngestJobSettingsCustomization::CustomizeDetails(IDetailLayoutBuilder& DetailLayout)
 {
-	TSharedRef<IPropertyHandle> PropertyHandle = DetailLayout.GetProperty(GET_MEMBER_NAME_CHECKED(UIngestJobSettingsBase, UploadHostName));
+	TSharedRef<IPropertyHandle> PropertyHandle = DetailLayout.GetProperty(GET_MEMBER_NAME_CHECKED(UIngestJobSettings, UploadHostName));
 	IDetailPropertyRow* PropertyRow = DetailLayout.EditDefaultProperty(PropertyHandle);
 	check(PropertyRow);
 
diff --git a/Engine/Plugins/VirtualProduction/CaptureManager/CaptureManagerApp/Source/CaptureManagerEditor/Private/UI/TakesView.cpp b/Engine/Plugins/VirtualProduction/CaptureManager/CaptureManagerApp/Source/CaptureManagerEditor/Private/UI/TakesView.cpp
--- a/Engine/Plugins/VirtualProduction/CaptureManager/CaptureManagerApp/Source/CaptureManagerEditor/Private/UI/TakesView.cpp
+++ b/Engine/Plugins/VirtualProduction/CaptureManager/CaptureManagerApp/Source/CaptureManagerEditor/Private/UI/TakesView.cpp
@@ -79,13 +79,6 @@ void STakesView::Construct(
 	OnRefreshTakes = InArgs._OnRefreshTakes;
 	CurrentPipeline = IngestPipelineManager->GetSelectedPipeline();
 
-	const TStrongObjectPtr<const UIngestJobSettingsBase> SelectedPreset(IngestJobSettingsManager->GetSelectedPreset());
-
-	if (SelectedPreset)
-	{
-		CurrentPresetName = SelectedPreset->DisplayName;
-	}
-
 	const TArray<UE::CaptureManager::FPipelineDetails>& Pipelines = IngestPipelineManager->GetPipelines();
 	for (const UE::CaptureManager::FPipelineDetails& Pipeline : Pipelines)
 	{
@@ -278,11 +271,11 @@ FReply STakesView::OnAddToQueueButtonClicked()
 			}
 		);
 
-		const TStrongObjectPtr<const UIngestJobSettingsBase> SelectedPreset(IngestJobSettingsManager->GetSelectedPreset());
+		const TStrongObjectPtr<const UIngestJobSettings> DefaultSettings(NewObject<UIngestJobSettings>());
 
-		if (SelectedPreset)
+		if (DefaultSettings)
 		{
-			OnAddTakesToIngestQueue.ExecuteIfBound(MoveTemp(SelectedTakeAssets), *SelectedPreset);
+			OnAddTakesToIngestQueue.ExecuteIfBound(MoveTemp(SelectedTakeAssets), *DefaultSettings);
 		}
 	}
 
@@ -385,7 +378,7 @@ EImageRotation FTakesView::ConvertOrientation(TOptional<FTakeMetadata::FVideo::E
 	}
 }
 
-void FTakesView::AddTakesToIngestQueue(const TArray<TObjectPtr<UTakeVirtualAsset>>& InTakeAssets, const UIngestJobSettingsBase& InJobPreset)
+void FTakesView::AddTakesToIngestQueue(const TArray<TObjectPtr<UTakeVirtualAsset>>& InTakeAssets, const UIngestJobSettings& InJobSettings)
 {
 	using namespace UE::CaptureManager;
 
@@ -404,26 +397,26 @@ void FTakesView::AddTakesToIngestQueue(const TArray<TObjectPtr<UTakeVirtualAsset
 	for (const TObjectPtr<UTakeVirtualAsset>& TakeAsset : InTakeAssets)
 	{
 		EImageRotation ImageRotationToApply =
-			TakeAsset->Metadata.Video.IsEmpty() ? InJobPreset.ImageRotation : ConvertOrientation(TakeAsset->Metadata.Video[0].Orientation);
+			TakeAsset->Metadata.Video.IsEmpty() ? InJobSettings.ImageRotation : ConvertOrientation(TakeAsset->Metadata.Video[0].Orientation);
 
 		// We convert the UObject based settings into something that the ingest job can truly own
 		FIngestJob::FSettings IngestJobSettings =
 		{
-			.WorkingDirectory = InJobPreset.WorkingDirectory.Path,
-			.DownloadFolder = InJobPreset.DownloadFolder.Path,
+			.WorkingDirectory = InJobSettings.WorkingDirectory.Path,
+			.DownloadFolder = InJobSettings.DownloadFolder.Path,
 			.VideoSettings =
 			{
-				.Format = InJobPreset.ImageFormat,
-				.FileNamePrefix = InJobPreset.ImageFileNamePrefix,
-				.ImagePixelFormat = InJobPreset.ImagePixelFormat,
+				.Format = InJobSettings.ImageFormat,
+				.FileNamePrefix = InJobSettings.ImageFileNamePrefix,
+				.ImagePixelFormat = InJobSettings.ImagePixelFormat,
 				.ImageRotation = ImageRotationToApply
 			},
 			.AudioSettings =
 			{
-				.Format = InJobPreset.AudioFormat,
-				.FileNamePrefix = InJobPreset.AudioFileNamePrefix
+				.Format = InJobSettings.AudioFormat,
+				.FileNamePrefix = InJobSettings.AudioFileNamePrefix
 			},
-			.UploadHostName = InJobPreset.UploadHostName
+			.UploadHostName = InJobSettings.UploadHostName
 		};
 
 		IngestJobs.Emplace(MakeShared<FIngestJob>(TakeAsset->CaptureDeviceId, TakeAsset->TakeId, TakeAsset->Metadata, Pipeline.PipelineConfig, MoveTemp(IngestJobSettings)));
diff --git a/Engine/Plugins/VirtualProduction/CaptureManager/CaptureManagerApp/Source/CaptureManagerEditor/Private/UI/TakesView.h b/Engine/Plugins/VirtualProduction/CaptureManager/CaptureManagerApp/Source/CaptureManagerEditor/Private/UI/TakesView.h
--- a/Engine/Plugins/VirtualProduction/CaptureManager/CaptureManagerApp/Source/CaptureManagerEditor/Private/UI/TakesView.h
+++ b/Engine/Plugins/VirtualProduction/CaptureManager/CaptureManagerApp/Source/CaptureManagerEditor/Private/UI/TakesView.h
@@ -35,7 +35,7 @@ class SIngestJobProcessor;
 class STakesView : public SCompoundWidget
 {
 public:
-	DECLARE_DELEGATE_TwoParams(FOnAddTakesToIngestQueue, const TArray<TObjectPtr<UTakeVirtualAsset>>&, const UIngestJobSettingsBase&);
+	DECLARE_DELEGATE_TwoParams(FOnAddTakesToIngestQueue, const TArray<TObjectPtr<UTakeVirtualAsset>>&, const UIngestJobSettings&);
 	DECLARE_DELEGATE(FOnRefreshTakes);
 
 	SLATE_BEGIN_ARGS(STakesView) {}
@@ -75,7 +75,6 @@ private:
 	TArray<TSharedRef<FText>> PipelineNames;
 	UE::CaptureManager::FPipelineDetails CurrentPipeline;
 	TSharedPtr<UE::CaptureManager::FIngestJobSettingsManager> IngestJobSettingsManager;
-	FText CurrentPresetName;
 };
 
 class FTakesView
@@ -103,7 +102,7 @@ private:
 	TArray<FAssetData> CreateTakeAssetsData(const TArray<TObjectPtr<UTakeVirtualAsset>>& InSelectedCaptureDevicesTakesArray);
 
 	EImageRotation ConvertOrientation(TOptional<FTakeMetadata::FVideo::EOrientation> InOrientation) const;
-	void AddTakesToIngestQueue(const TArray<TObjectPtr<UTakeVirtualAsset>>& InTakeAssets, const UIngestJobSettingsBase& InJobPreset);
+	void AddTakesToIngestQueue(const TArray<TObjectPtr<UTakeVirtualAsset>>& InTakeAssets, const UIngestJobSettings& InJobPreset);
 	void Refresh();
 	void RefreshSingleTake(TObjectPtr<ULiveLinkDevice> InDevice);
 	void RefreshAllTakes();
