[CaptureManager] Fix crashes when removing devices during ingest

diff --git a/Engine/Plugins/VirtualProduction/CaptureManager/CaptureManagerApp/Source/CaptureManagerEditor/Private/IngestManagement/IngestJob.h b/Engine/Plugins/VirtualProduction/CaptureManager/CaptureManagerApp/Source/CaptureManagerEditor/Private/IngestManagement/IngestJob.h
--- a/Engine/Plugins/VirtualProduction/CaptureManager/CaptureManagerApp/Source/CaptureManagerEditor/Private/IngestManagement/IngestJob.h
+++ b/Engine/Plugins/VirtualProduction/CaptureManager/CaptureManagerApp/Source/CaptureManagerEditor/Private/IngestManagement/IngestJob.h
@@ -13,12 +13,12 @@ namespace UE::CaptureManager
 class FIngestJob
 {
 public:
-	enum class EProcessingState
+	enum class EProcessingState : uint32
 	{
-		Pending,
-		Running,
-		Complete,
-		Aborted
+		Pending = 1,
+		Running = 2,
+		Complete = 4,
+		Aborted = 8
 	};
 
 	struct FSettings
diff --git a/Engine/Plugins/VirtualProduction/CaptureManager/CaptureManagerApp/Source/CaptureManagerEditor/Private/IngestManagement/IngestJobExecutor.cpp b/Engine/Plugins/VirtualProduction/CaptureManager/CaptureManagerApp/Source/CaptureManagerEditor/Private/IngestManagement/IngestJobExecutor.cpp
--- a/Engine/Plugins/VirtualProduction/CaptureManager/CaptureManagerApp/Source/CaptureManagerEditor/Private/IngestManagement/IngestJobExecutor.cpp
+++ b/Engine/Plugins/VirtualProduction/CaptureManager/CaptureManagerApp/Source/CaptureManagerEditor/Private/IngestManagement/IngestJobExecutor.cpp
@@ -91,20 +91,23 @@ uint32 FIngestJobExecutor::Run()
 			check(Job->GetProcessingState() != FIngestJob::EProcessingState::Pending);
 			JobProcessingStateChanged.ExecuteIfBound(Job->GetGuid(), Job->GetProcessingState());
 
+			FScopeLock Lock(&CriticalSection);
+
 			// Note: This may be a blocking call depending on how the capture device has been implemented. If so, this 
 			// may take a while to stop (no way to cancel in that case).
 			TUniquePtr<FIngestJobExecutor::FRunIngestContext> RunIngestContext = RunIngest(*Job, Promise);
-
-			FScopeLock Lock(&CriticalSection);
-			check(RunIngestContext);
-
-			RunIngestProcessContext = MoveTemp(RunIngestContext);
-
-			// There's a period between RunIngest and the assignment of the member context, during which Stop() may have
-			// been called. If so, immediately terminate (as Cancel() will not be called otherwise).
-			if (bStopRequested && !Future.IsReady() && RunIngestProcessContext)
+			
+			// Check the context in case the device has been removed during the RunIngest call
+			if (RunIngestContext)
 			{
-				ILiveLinkDeviceCapability_Ingest::Execute_CancelIngestProcess(RunIngestProcessContext->Device.Get(), RunIngestProcessContext->ProcessHandle.Get());
+				RunIngestProcessContext = MoveTemp(RunIngestContext);
+
+				// There's a period between RunIngest and the assignment of the member context, during which Stop() may have
+				// been called. If so, immediately terminate (as Cancel() will not be called otherwise).
+				if (bStopRequested && !Future.IsReady() && RunIngestProcessContext)
+				{
+					ILiveLinkDeviceCapability_Ingest::Execute_CancelIngestProcess(RunIngestProcessContext->Device.Get(), RunIngestProcessContext->ProcessHandle.Get());
+				}
 			}
 		}
 		else
@@ -146,6 +149,16 @@ void FIngestJobExecutor::Stop()
 	}
 }
 
+void FIngestJobExecutor::CancelForDevice(const FGuid& InDeviceId)
+{
+	FScopeLock Lock(&CriticalSection);
+
+	if (RunIngestProcessContext.IsValid() && RunIngestProcessContext->Device.IsValid() && RunIngestProcessContext->Device->GetDeviceId() == InDeviceId)
+	{
+		ILiveLinkDeviceCapability_Ingest::Execute_CancelIngestProcess(RunIngestProcessContext->Device.Get(), RunIngestProcessContext->ProcessHandle.Get());
+	}
+}
+
 void FIngestJobExecutor::Exit()
 {
 	bIsRunning = false;
diff --git a/Engine/Plugins/VirtualProduction/CaptureManager/CaptureManagerApp/Source/CaptureManagerEditor/Private/IngestManagement/IngestJobExecutor.h b/Engine/Plugins/VirtualProduction/CaptureManager/CaptureManagerApp/Source/CaptureManagerEditor/Private/IngestManagement/IngestJobExecutor.h
--- a/Engine/Plugins/VirtualProduction/CaptureManager/CaptureManagerApp/Source/CaptureManagerEditor/Private/IngestManagement/IngestJobExecutor.h
+++ b/Engine/Plugins/VirtualProduction/CaptureManager/CaptureManagerApp/Source/CaptureManagerEditor/Private/IngestManagement/IngestJobExecutor.h
@@ -24,6 +24,8 @@ public:
 	virtual void Stop() override;
 	virtual void Exit() override;
 
+	void CancelForDevice(const FGuid& InDeviceId);
+
 private:
 	struct FRunIngestContext
 	{
diff --git a/Engine/Plugins/VirtualProduction/CaptureManager/CaptureManagerApp/Source/CaptureManagerEditor/Private/IngestManagement/IngestJobProcessor.cpp b/Engine/Plugins/VirtualProduction/CaptureManager/CaptureManagerApp/Source/CaptureManagerEditor/Private/IngestManagement/IngestJobProcessor.cpp
--- a/Engine/Plugins/VirtualProduction/CaptureManager/CaptureManagerApp/Source/CaptureManagerEditor/Private/IngestManagement/IngestJobProcessor.cpp
+++ b/Engine/Plugins/VirtualProduction/CaptureManager/CaptureManagerApp/Source/CaptureManagerEditor/Private/IngestManagement/IngestJobProcessor.cpp
@@ -168,6 +168,38 @@ bool FIngestJobProcessor::SetJobSettings(const FGuid& JobGuid, FIngestJob::FSett
 	return ProcessingQueue->SetJobSettings(JobGuid, MoveTemp(InSettings));
 }
 
+uint32 FIngestJobProcessor::CountQueuedDeviceJobs(const FGuid InDeviceId, const uint32 InJobsToCountFlags)
+{
+	return ProcessingQueue->CountQueuedDeviceJobs(InDeviceId, InJobsToCountFlags);
+}
+
+int32 FIngestJobProcessor::RemoveJobsForDevice(const FGuid InDeviceId)
+{
+	TArray<FGuid> JobsGuidsRemoved = ProcessingQueue->RemoveJobsForDevice(InDeviceId);
+
+	if (JobsGuidsRemoved.Num() > 0)
+	{
+		JobsRemoved.ExecuteIfBound(JobsGuidsRemoved);
+	}
+
+	return JobsGuidsRemoved.Num();
+}
+
+void FIngestJobProcessor::StopProcessingForDevice(const FGuid& InDeviceId)
+{
+	if (!IsProcessing())
+	{
+		return;
+	}
+
+	FScopeLock Lock(&CriticalSection);
+
+	for (const TUniquePtr<FIngestJobExecutor>& Executor : Executors)
+	{
+		Executor->CancelForDevice(InDeviceId);
+	}
+}
+
 int32 FIngestJobProcessor::RemoveJobs(const TArray<FGuid>& JobGuidsToRemove)
 {
 	if (JobGuidsToRemove.IsEmpty())
diff --git a/Engine/Plugins/VirtualProduction/CaptureManager/CaptureManagerApp/Source/CaptureManagerEditor/Private/IngestManagement/IngestJobProcessor.h b/Engine/Plugins/VirtualProduction/CaptureManager/CaptureManagerApp/Source/CaptureManagerEditor/Private/IngestManagement/IngestJobProcessor.h
--- a/Engine/Plugins/VirtualProduction/CaptureManager/CaptureManagerApp/Source/CaptureManagerEditor/Private/IngestManagement/IngestJobProcessor.h
+++ b/Engine/Plugins/VirtualProduction/CaptureManager/CaptureManagerApp/Source/CaptureManagerEditor/Private/IngestManagement/IngestJobProcessor.h
@@ -31,12 +31,15 @@ public:
 	~FIngestJobProcessor();
 
 	int32 AddJobs(TArray<TSharedRef<FIngestJob>> IngestJobs);
+	uint32 CountQueuedDeviceJobs(const FGuid InDeviceId, const uint32 InJobsToCountFlags);
+	int32 RemoveJobsForDevice(const FGuid InDeviceId);
 	int32 RemoveJobs(const TArray<FGuid>& JobGuidsToRemove);
 	int32 RemoveAllJobs();
 	bool SetJobSettings(const FGuid& JobGuid, FIngestJob::FSettings Settings);
 
 	void StartProcessing();
 	void StopProcessing();
+	void StopProcessingForDevice(const FGuid& InDeviceId);
 
 	bool IsProcessing() const;
 	bool IsStopping() const;
diff --git a/Engine/Plugins/VirtualProduction/CaptureManager/CaptureManagerApp/Source/CaptureManagerEditor/Private/IngestManagement/IngestJobQueue.cpp b/Engine/Plugins/VirtualProduction/CaptureManager/CaptureManagerApp/Source/CaptureManagerEditor/Private/IngestManagement/IngestJobQueue.cpp
--- a/Engine/Plugins/VirtualProduction/CaptureManager/CaptureManagerApp/Source/CaptureManagerEditor/Private/IngestManagement/IngestJobQueue.cpp
+++ b/Engine/Plugins/VirtualProduction/CaptureManager/CaptureManagerApp/Source/CaptureManagerEditor/Private/IngestManagement/IngestJobQueue.cpp
@@ -50,6 +50,39 @@ TArray<FGuid> FIngestJobQueue::RemoveAll()
 
 	return GuidsRemoved;
 }
+uint32 FIngestJobQueue::CountQueuedDeviceJobs(const FGuid InDeviceId, const uint32 InJobsToCountFlags)
+{
+	for (const TSharedRef<FIngestJob>& Job : Jobs)
+	{
+		if (Job->GetCaptureDeviceId() == InDeviceId)
+		{
+			if (static_cast<uint32>(Job->GetProcessingState()) & InJobsToCountFlags)
+			{
+				return true;
+			}
+		}
+	}
+	return false;
+}
+
+TArray<FGuid> FIngestJobQueue::RemoveJobsForDevice(const FGuid InDeviceId)
+{
+	TArray<FGuid> GuidsRemoved;
+
+	const int32 NumRemoved = Jobs.RemoveAll(
+		[&InDeviceId, &GuidsRemoved](const TSharedRef<FIngestJob>& Job)
+		{
+			bool IsDeviceJob = Job->GetCaptureDeviceId() == InDeviceId;
+			if (IsDeviceJob)
+			{
+				GuidsRemoved.Add(Job->GetGuid());
+			}
+			
+			return IsDeviceJob;
+		}
+	);
+	return GuidsRemoved;
+}
 
 TSharedPtr<FIngestJob> FIngestJobQueue::ClaimFirstPending(const FIngestJob::EProcessingState ProcessingState)
 {
diff --git a/Engine/Plugins/VirtualProduction/CaptureManager/CaptureManagerApp/Source/CaptureManagerEditor/Private/IngestManagement/IngestJobQueue.h b/Engine/Plugins/VirtualProduction/CaptureManager/CaptureManagerApp/Source/CaptureManagerEditor/Private/IngestManagement/IngestJobQueue.h
--- a/Engine/Plugins/VirtualProduction/CaptureManager/CaptureManagerApp/Source/CaptureManagerEditor/Private/IngestManagement/IngestJobQueue.h
+++ b/Engine/Plugins/VirtualProduction/CaptureManager/CaptureManagerApp/Source/CaptureManagerEditor/Private/IngestManagement/IngestJobQueue.h
@@ -24,6 +24,8 @@ public:
 
 	TArray<FGuid> Remove(TArray<FGuid> GuidsToRemove);
 	TArray<FGuid> RemoveAll();
+	uint32 CountQueuedDeviceJobs(const FGuid InDeviceId, const uint32 InJobsToCountFlags);
+	TArray<FGuid> RemoveJobsForDevice(const FGuid InDeviceId);
 
 	bool SetJobSettings(const FGuid& JobGuid, FIngestJob::FSettings InSettings);
 
diff --git a/Engine/Plugins/VirtualProduction/CaptureManager/CaptureManagerApp/Source/CaptureManagerEditor/Private/UI/CaptureManagerPanelController.cpp b/Engine/Plugins/VirtualProduction/CaptureManager/CaptureManagerApp/Source/CaptureManagerEditor/Private/UI/CaptureManagerPanelController.cpp
--- a/Engine/Plugins/VirtualProduction/CaptureManager/CaptureManagerApp/Source/CaptureManagerEditor/Private/UI/CaptureManagerPanelController.cpp
+++ b/Engine/Plugins/VirtualProduction/CaptureManager/CaptureManagerApp/Source/CaptureManagerEditor/Private/UI/CaptureManagerPanelController.cpp
@@ -236,7 +236,7 @@ void FCaptureManagerPanelController::OnDeviceRemoved(FGuid InDeviceId, ULiveLink
 		Delegate->ConnectionChanged.RemoveAll(this);
 	}
 
-	TakesView->CaptureDeviceToBeRemoved(InDevice);
+	TakesView->CaptureDeviceRemoved(InDevice);
 	SourcesReachableMap.Remove(InDeviceId);
 }
 
@@ -253,6 +253,7 @@ void FCaptureManagerPanelController::OnReachableEvent(ELiveLinkDeviceConnectionS
 	}
 	else
 	{
+		IngestJobProcessorWidget->RemoveJobsForDevice(InDeviceId);
 		TakesView->CaptureDeviceStopped(InDeviceId);
 	}
 }
diff --git a/Engine/Plugins/VirtualProduction/CaptureManager/CaptureManagerApp/Source/CaptureManagerEditor/Private/UI/SIngestJobProcessor.cpp b/Engine/Plugins/VirtualProduction/CaptureManager/CaptureManagerApp/Source/CaptureManagerEditor/Private/UI/SIngestJobProcessor.cpp
--- a/Engine/Plugins/VirtualProduction/CaptureManager/CaptureManagerApp/Source/CaptureManagerEditor/Private/UI/SIngestJobProcessor.cpp
+++ b/Engine/Plugins/VirtualProduction/CaptureManager/CaptureManagerApp/Source/CaptureManagerEditor/Private/UI/SIngestJobProcessor.cpp
@@ -571,6 +571,21 @@ bool SIngestJobProcessor::SetJobSettings(const FGuid& JobGuid, FIngestJob::FSett
 	return IngestJobProcessor->SetJobSettings(JobGuid, MoveTemp(JobSettings));
 }
 
+uint32 SIngestJobProcessor::CountQueuedDeviceJobs(const FGuid InDeviceId, const uint32 InJobsToCountFlags)
+{
+	return IngestJobProcessor->CountQueuedDeviceJobs(InDeviceId, InJobsToCountFlags);
+}
+
+int32 SIngestJobProcessor::RemoveJobsForDevice(const FGuid InDeviceId)
+{
+	return IngestJobProcessor->RemoveJobsForDevice(InDeviceId);
+}
+
+void SIngestJobProcessor::Stop(const FGuid InDeviceId)
+{
+	IngestJobProcessor->StopProcessingForDevice(InDeviceId);
+}
+
 } // namespace UE::CaptureManager
 
 #undef LOCTEXT_NAMESPACE
diff --git a/Engine/Plugins/VirtualProduction/CaptureManager/CaptureManagerApp/Source/CaptureManagerEditor/Private/UI/SIngestJobProcessor.h b/Engine/Plugins/VirtualProduction/CaptureManager/CaptureManagerApp/Source/CaptureManagerEditor/Private/UI/SIngestJobProcessor.h
--- a/Engine/Plugins/VirtualProduction/CaptureManager/CaptureManagerApp/Source/CaptureManagerEditor/Private/UI/SIngestJobProcessor.h
+++ b/Engine/Plugins/VirtualProduction/CaptureManager/CaptureManagerApp/Source/CaptureManagerEditor/Private/UI/SIngestJobProcessor.h
@@ -30,6 +30,9 @@ public:
 	bool IsProcessing() const;
 	int32 AddJobs(TArray<TSharedRef<FIngestJob>> Jobs);
 	bool SetJobSettings(const FGuid& JobGuid, FIngestJob::FSettings JobSettings);
+	uint32 CountQueuedDeviceJobs(const FGuid InDeviceId, const uint32 InJobsToCountFlags);
+	int32 RemoveJobsForDevice(const FGuid InDeviceId);
+	void Stop(const FGuid InDeviceId);
 
 	FSelectionChanged& OnSelectionChanged();
 	FJobsAdded& OnJobsAdded();
diff --git a/Engine/Plugins/VirtualProduction/CaptureManager/CaptureManagerApp/Source/CaptureManagerEditor/Private/UI/TakesView.cpp b/Engine/Plugins/VirtualProduction/CaptureManager/CaptureManagerApp/Source/CaptureManagerEditor/Private/UI/TakesView.cpp
--- a/Engine/Plugins/VirtualProduction/CaptureManager/CaptureManagerApp/Source/CaptureManagerEditor/Private/UI/TakesView.cpp
+++ b/Engine/Plugins/VirtualProduction/CaptureManager/CaptureManagerApp/Source/CaptureManagerEditor/Private/UI/TakesView.cpp
@@ -16,6 +16,7 @@
 #include "Widgets/Input/SComboBox.h"
 #include "SPositiveActionButton.h"
 #include "Widgets/Images/SThrobber.h"
+#include "Misc/MessageDialog.h"
 
 #include "Ingest/LiveLinkDeviceCapability_Ingest.h"
 #include "Ingest/IngestCapability_Events.h"
@@ -437,6 +438,45 @@ void FTakesView::Refresh()
 {
 	using namespace UE::CaptureManager;
 
+	TSharedPtr<FCaptureManagerPanelController> Controller = PanelController.Pin();
+	if (Controller)
+	{	
+		bool bHasQueuedJobs = false;
+		const uint32 JobsToCount = static_cast<uint32>(UE::CaptureManager::FIngestJob::EProcessingState::Pending)
+								 | static_cast<uint32>(UE::CaptureManager::FIngestJob::EProcessingState::Running);
+
+		TSharedRef<UE::CaptureManager::SIngestJobProcessor> IngestJobProcessor = Controller->GetIngestJobProcessorWidget();
+		if (SelectedItem.IsValid() && SelectedItem->Implements<ULiveLinkDeviceCapability_Ingest>())
+		{
+			const FGuid DeviceId = SelectedItem->GetDeviceId();
+			bHasQueuedJobs = IngestJobProcessor->CountQueuedDeviceJobs(DeviceId, JobsToCount) > 0;
+		}
+		else
+		{
+
+			for (const TObjectPtr<ULiveLinkDevice>& Device : Controller->GetCaptureDevices())
+			{
+				if (IngestJobProcessor->CountQueuedDeviceJobs(Device->GetDeviceId(), JobsToCount) > 0)
+				{
+					bHasQueuedJobs = true;
+					break;
+				}
+			}
+		}
+
+		if (bHasQueuedJobs)
+		{
+			const FText ConfirmationTitle = LOCTEXT("ClearAffectedDeviceJobsTitle", "Remove queued takes?");
+			const FText ConfirmationMessage = LOCTEXT("ClearAffectedDeviceJobsMessage", "Refreshing will remove takes queued for ingest and cancel in progress ingest jobs. Do you wish to proceed with the refresh?");
+			EAppReturnType::Type Answer = FMessageDialog::Open(EAppMsgType::YesNo, ConfirmationMessage, ConfirmationTitle);
+			if (Answer == EAppReturnType::No)
+			{
+				// Give up refreshing
+				return;
+			}
+		}
+	}
+
 	if (SelectedItem.IsValid() && SelectedItem->Implements<ULiveLinkDeviceCapability_Ingest>())
 	{
 		RefreshSingleTake(SelectedItem.Get());
@@ -501,20 +541,18 @@ void FTakesView::UpdateTakeListCallback(TArray<UE::CaptureManager::FTakeId> InTa
 		RefreshTakesViewDelegate.ExecuteIfBound(true);
 		return;
 	}
-
-	bool bAreEqual = EqualTo(TakeAssets, InTakeIds, 
-							 [](const TUniqueObjectPtr<UTakeVirtualAsset>& InLeft, FTakeId InTakeId)
-							 {
-								 return InLeft->TakeId == InTakeId;
-							 });
-
-	check(bAreEqual);
 }
 
 void FTakesView::ResetTakesCache(FGuid InCaptureDeviceId)
 {
 	if (CaptureDeviceTakesMap.Contains(InCaptureDeviceId))
 	{
+		TSharedPtr<FCaptureManagerPanelController> Controller = PanelController.Pin();
+		if (Controller)
+		{
+			TSharedRef<UE::CaptureManager::SIngestJobProcessor> IngestJobProcessor = Controller->GetIngestJobProcessorWidget();
+			IngestJobProcessor->RemoveJobsForDevice(InCaptureDeviceId);
+		}
 		CaptureDeviceTakesMap[InCaptureDeviceId].Empty();
 	}
 	else
@@ -769,14 +807,28 @@ void FTakesView::CaptureDeviceAdded(ULiveLinkDevice* InDevice)
 	});
 }
 
-void FTakesView::CaptureDeviceToBeRemoved(ULiveLinkDevice* InDevice)
+void FTakesView::CaptureDeviceRemoved(ULiveLinkDevice* InDevice)
 {
+	// ULiveLinkDevice has been removed so handle device removal in Capture Manager (e.g. removal of ingest jobs)
 	AsyncTask(ENamedThreads::GameThread, [this, Controller = PanelController.Pin(), Device = TStrongObjectPtr<ULiveLinkDevice>(InDevice)]()
 		{
 			TScriptInterface<ILiveLinkDeviceCapability_Ingest> CapabilityInterface(Device.Get());
 
 			CapabilityInterface->UnsubscribeAll();
 
+			TSharedRef<UE::CaptureManager::SIngestJobProcessor> IngestJobProcessor = Controller->GetIngestJobProcessorWidget();
+
+			// There's a period between RunIngest and the assignment of the member context, during which Stop() may have
+			// been called. If so, immediately terminate (as Cancel() will not be called otherwise).
+			IngestJobProcessor->Stop(Device->GetDeviceId());
+			
+
+			int32 NumJobsRemoved = IngestJobProcessor->RemoveJobsForDevice(Device->GetDeviceId());
+			if (NumJobsRemoved > 0)
+			{
+				UE_LOG(LogTakesView, Display, TEXT("Device '%s' removed. Removed %d corresponding ingest jobs from the queue."), *Device->GetDisplayName().ToString(), NumJobsRemoved);
+			}
+			
 			CaptureDeviceTakesMap.Remove(Device->GetDeviceId());
 			RefreshTakesViewDelegate.ExecuteIfBound(true);
 		});
diff --git a/Engine/Plugins/VirtualProduction/CaptureManager/CaptureManagerApp/Source/CaptureManagerEditor/Private/UI/TakesView.h b/Engine/Plugins/VirtualProduction/CaptureManager/CaptureManagerApp/Source/CaptureManagerEditor/Private/UI/TakesView.h
--- a/Engine/Plugins/VirtualProduction/CaptureManager/CaptureManagerApp/Source/CaptureManagerEditor/Private/UI/TakesView.h
+++ b/Engine/Plugins/VirtualProduction/CaptureManager/CaptureManagerApp/Source/CaptureManagerEditor/Private/UI/TakesView.h
@@ -85,7 +85,7 @@ public:
 	void CaptureDeviceStarted(FGuid InCaptureDeviceId);
 	void CaptureDeviceStopped(FGuid InCaptureDeviceId);
 	void CaptureDeviceAdded(ULiveLinkDevice* InDevice);
-	void CaptureDeviceToBeRemoved(ULiveLinkDevice* InDevice);
+	void CaptureDeviceRemoved(ULiveLinkDevice* InDevice);
 
 	TSharedPtr<STakesView> TakesTileView;
 private:
