Control flow crash during ingest

diff --git a/Engine/Plugins/VirtualProduction/CaptureManager/CaptureManagerApp/Source/CaptureDataConverter/Private/CaptureDataConverter.cpp b/Engine/Plugins/VirtualProduction/CaptureManager/CaptureManagerApp/Source/CaptureDataConverter/Private/CaptureDataConverter.cpp
--- a/Engine/Plugins/VirtualProduction/CaptureManager/CaptureManagerApp/Source/CaptureDataConverter/Private/CaptureDataConverter.cpp
+++ b/Engine/Plugins/VirtualProduction/CaptureManager/CaptureManagerApp/Source/CaptureDataConverter/Private/CaptureDataConverter.cpp
@@ -194,7 +194,10 @@ FCaptureDataConverterResult<void> FCaptureDataConverter::Run(FCaptureDataConvert
 
 void FCaptureDataConverter::Cancel()
 {
-	StopRequester.RequestStop();
+	if (!StopRequester.IsStopRequested())
+	{
+		StopRequester.RequestStop();
 
-	Pipeline->Cancel();
+		Pipeline->Cancel();
+	}
 }
\ No newline at end of file
diff --git a/Engine/Plugins/VirtualProduction/CaptureManager/CaptureManagerApp/Source/CaptureManagerPipeline/Private/CaptureManagerPipeline.cpp b/Engine/Plugins/VirtualProduction/CaptureManager/CaptureManagerApp/Source/CaptureManagerPipeline/Private/CaptureManagerPipeline.cpp
--- a/Engine/Plugins/VirtualProduction/CaptureManager/CaptureManagerApp/Source/CaptureManagerPipeline/Private/CaptureManagerPipeline.cpp
+++ b/Engine/Plugins/VirtualProduction/CaptureManager/CaptureManagerApp/Source/CaptureManagerPipeline/Private/CaptureManagerPipeline.cpp
@@ -100,7 +100,7 @@ FGuid FCaptureManagerPipeline::AddSyncedNode(TSharedPtr<FCaptureManagerPipelineN
 {
 	FGuid UniqueId = FGuid::NewGuid();
 
-	SyncNodes.Add(UniqueId, MoveTemp(InNode));
+	SyncNodes.Lock()->Add(UniqueId, MoveTemp(InNode));
 
 	return UniqueId;
 }
@@ -109,7 +109,7 @@ FGuid FCaptureManagerPipeline::AddParallelPipelineNode(TSharedPtr<FCaptureManage
 {
 	FGuid UniqueId = FGuid::NewGuid();
 
-	ParallelNodes.Add(UniqueId, MoveTemp(InNode));
+	ParallelNodes.Lock()->Add(UniqueId, MoveTemp(InNode));
 
 	return UniqueId;
 }
@@ -127,10 +127,16 @@ FCaptureManagerPipeline::FResult FCaptureManagerPipeline::Run()
 		FText Message = LOCTEXT("Run_Canceled", "The pipeline has been canceled");
 		int32 Code = -10;
 
-		TArray<FGuid> NodeIds;
-		ParallelNodes.GetKeys(NodeIds);
+		TArray<FGuid> ParellelNodeIds, SyncNodeIds;
+		ParallelNodes.Lock()->GetKeys(ParellelNodeIds);
+		SyncNodes.Lock()->GetKeys(SyncNodeIds);
+		
+		for (const FGuid& UniqueId : ParellelNodeIds)
+		{
+			ThreadSafeResult.Lock()->Add(UniqueId, MakeError(MoveTemp(Message), Code));
+		}
 
-		for (const FGuid& UniqueId : NodeIds)
+		for (const FGuid& UniqueId : SyncNodeIds)
 		{
 			ThreadSafeResult.Lock()->Add(UniqueId, MakeError(MoveTemp(Message), Code));
 		}
@@ -149,11 +155,12 @@ FCaptureManagerPipeline::FResult FCaptureManagerPipeline::Run()
 		int32 Index = 0;
 
 		TArray<FGuid> NodeIds;
-		ParallelNodes.GetKeys(NodeIds);
+		FMonitor<FNodeMap>::FHelper ParallelNodesHelper = ParallelNodes.Lock(); 
+		ParallelNodesHelper->GetKeys(NodeIds);
 
 		for (const FGuid& UniqueId : NodeIds)
 		{
-			TSharedPtr<FCaptureManagerPipelineNode> Node = ParallelNodes.FindChecked(UniqueId);
+			TSharedPtr<FCaptureManagerPipelineNode> Node = ParallelNodesHelper->FindChecked(UniqueId);
 
 			FString Name = Node->GetName();
 			InConcurrentFlow->AddOrGetFlow(Index, Name)
@@ -162,13 +169,18 @@ FCaptureManagerPipeline::FResult FCaptureManagerPipeline::Run()
 				FCaptureManagerPipelineNode::FResult Result = ExeNode->Execute();
 				ThreadSafeResult.Lock()->Add(MoveTemp(Id), Result);
 
-				ParallelNodes.Remove(Id);
+				ParallelNodes.Lock()->Remove(Id);
+
+				if (InSubFlow->HasCancelBeenRequested())
+				{
+					return;
+				}
 
 				if (Result.HasValue())
 				{
 					InSubFlow->ContinueFlow();
 				}
-				else if (Result.HasError() && !InSubFlow->HasCancelBeenRequested())
+				else
 				{
 					InSubFlow->CancelFlow();
 				}
@@ -178,12 +190,13 @@ FCaptureManagerPipeline::FResult FCaptureManagerPipeline::Run()
 		}
 	});
 
+	FMonitor<FNodeMap>::FHelper SyncNodesHelper = SyncNodes.Lock();
 	TArray<FGuid> SyncedNodeIds;
-	SyncNodes.GetKeys(SyncedNodeIds);
+	SyncNodesHelper->GetKeys(SyncedNodeIds);
 
 	for (const FGuid& UniqueId : SyncedNodeIds)
 	{
-		TSharedPtr<FCaptureManagerPipelineNode> Node = SyncNodes.FindChecked(UniqueId);
+		TSharedPtr<FCaptureManagerPipelineNode> Node = SyncNodesHelper->FindChecked(UniqueId);
 		FString Name = Node->GetName();
 
 		Impl->QueueWait(Name).BindLambda([this, Id = UniqueId, ExeNode = MoveTemp(Node), &ThreadSafeResult](FControlFlowNodeRef InSubFlow) mutable
@@ -191,13 +204,18 @@ FCaptureManagerPipeline::FResult FCaptureManagerPipeline::Run()
 			FCaptureManagerPipelineNode::FResult Result = ExeNode->Execute();
 			ThreadSafeResult.Lock()->Add(MoveTemp(Id), Result);
 
-			SyncNodes.Remove(Id);
+			SyncNodes.Lock()->Remove(Id);
+
+			if (InSubFlow->HasCancelBeenRequested())
+			{
+				return;
+			}
 
 			if (Result.HasValue())
 			{
 				InSubFlow->ContinueFlow();
 			}
-			else if (Result.HasError() && !InSubFlow->HasCancelBeenRequested())
+			else
 			{
 				InSubFlow->CancelFlow();
 			}
@@ -215,12 +233,12 @@ void FCaptureManagerPipeline::Cancel()
 {
 	if (Impl->IsRunning())
 	{
-		for (const TPair<FGuid, TSharedPtr<FCaptureManagerPipelineNode>>& NodePair : ParallelNodes)
+		for (const TPair<FGuid, TSharedPtr<FCaptureManagerPipelineNode>>& NodePair : *ParallelNodes.Lock())
 		{
 			NodePair.Value->Cancel();
 		}
 
-		for (const TPair<FGuid, TSharedPtr<FCaptureManagerPipelineNode>>& NodePair : SyncNodes)
+		for (const TPair<FGuid, TSharedPtr<FCaptureManagerPipelineNode>>& NodePair : *SyncNodes.Lock())
 		{
 			NodePair.Value->Cancel();
 		}
diff --git a/Engine/Plugins/VirtualProduction/CaptureManager/CaptureManagerApp/Source/CaptureManagerPipeline/Private/CaptureManagerPipelineNode.cpp b/Engine/Plugins/VirtualProduction/CaptureManager/CaptureManagerApp/Source/CaptureManagerPipeline/Private/CaptureManagerPipelineNode.cpp
--- a/Engine/Plugins/VirtualProduction/CaptureManager/CaptureManagerApp/Source/CaptureManagerPipeline/Private/CaptureManagerPipelineNode.cpp
+++ b/Engine/Plugins/VirtualProduction/CaptureManager/CaptureManagerApp/Source/CaptureManagerPipeline/Private/CaptureManagerPipelineNode.cpp
@@ -67,12 +67,17 @@ private:
 	{
 		FResult& Result = InResult.Get();
 		Result = InStep.Execute();
+		
+		if (InFlowHandle->HasCancelBeenRequested())
+		{
+			return;
+		}
 
 		if (Result.HasValue())
 		{
 			InFlowHandle->ContinueFlow();
 		}
-		else if (Result.HasError() && !InFlowHandle->HasCancelBeenRequested())
+		else
 		{
 			InFlowHandle->CancelFlow();
 		}
diff --git a/Engine/Plugins/VirtualProduction/CaptureManager/CaptureManagerApp/Source/CaptureManagerPipeline/Public/CaptureManagerPipeline.h b/Engine/Plugins/VirtualProduction/CaptureManager/CaptureManagerApp/Source/CaptureManagerPipeline/Public/CaptureManagerPipeline.h
--- a/Engine/Plugins/VirtualProduction/CaptureManager/CaptureManagerApp/Source/CaptureManagerPipeline/Public/CaptureManagerPipeline.h
+++ b/Engine/Plugins/VirtualProduction/CaptureManager/CaptureManagerApp/Source/CaptureManagerPipeline/Public/CaptureManagerPipeline.h
@@ -12,6 +12,8 @@
 #include "Templates/SharedPointer.h"
 #include "Misc/Guid.h"
 
+#include "Async/Monitor.h"
+
 enum class EPipelineExecutionPolicy
 {
 	Asynchronous = 0,
@@ -45,8 +47,11 @@ private:
 
 	TPimplPtr<class FCaptureManagerPipelineImpl> Impl;
 
-	TMap<FGuid, TSharedPtr<FCaptureManagerPipelineNode>> ParallelNodes;
-	TMap<FGuid, TSharedPtr<FCaptureManagerPipelineNode>> SyncNodes;
+	
+	using FNodeMap = TMap<FGuid, TSharedPtr<FCaptureManagerPipelineNode>>;
+
+	UE::CaptureManager::FMonitor<FNodeMap> ParallelNodes;
+	UE::CaptureManager::FMonitor<FNodeMap> SyncNodes;
 
 	EPipelineExecutionPolicy ExecutionPolicy;
 };
\ No newline at end of file
