CaptureManager: Fixed LLF video timecode when the first frame is fractional and rounded up to an invalid frame number.

#rnx

diff --git a/Engine/Plugins/VirtualProduction/CaptureManager/CaptureManagerApp/Source/LiveLinkFaceMetadata/Private/LiveLinkFaceMetadata.cpp b/Engine/Plugins/VirtualProduction/CaptureManager/CaptureManagerApp/Source/LiveLinkFaceMetadata/Private/LiveLinkFaceMetadata.cpp
--- a/Engine/Plugins/VirtualProduction/CaptureManager/CaptureManagerApp/Source/LiveLinkFaceMetadata/Private/LiveLinkFaceMetadata.cpp
+++ b/Engine/Plugins/VirtualProduction/CaptureManager/CaptureManagerApp/Source/LiveLinkFaceMetadata/Private/LiveLinkFaceMetadata.cpp
@@ -822,7 +822,11 @@ public:
 						continue;
 					}
 
-					VideoTimecode = Timecode;
+					// We re-construct the timecode using the frame rate to resolve any invalid frame number rounding
+					// which may have occurred while parsing the frame log. As an example: a fractional frame number of
+					// 59.780 would round to 60, which is an invalid frame number at 60FPS. Going to/from a frame number
+					// fixes the problem.
+					VideoTimecode = FTimecode::FromFrameNumber(Timecode.ToFrameNumber(OriginalFrameRate), OriginalFrameRate);
 					VideoTimecodeRate = OriginalFrameRate;
 				}
 			}
@@ -838,10 +842,13 @@ public:
 
 					// Depth timecode rate and timecode are 30fps, and there is no need for conversion
 					FFrameRate TargetFrameRate{ 30, 1 };
+
+					// This step is also needed to resolve any invalid frame number rounding which may have occurred
+					// during the parsing of the frame log (as it implicitly does a conversion to/from frame number).
 					DepthTimecode = FTimecode::FromTimespan(Timecode.ToTimespan(OriginalFrameRate), TargetFrameRate, true);
 					DepthTimecodeRate = MoveTemp(TargetFrameRate);
 				}
-				
+
 				++DepthFrameCount;
 			}
 			else if (LogEntry.EntryType() == FrameLogEntry::AudioType)
@@ -864,6 +871,9 @@ public:
 					// This probably isn't mandatory and we should at some point step away from this when we make
 					// sure this assumption is not present in other parts of the MHA codebase.
 					FFrameRate TargetFrameRate{ 30, 1 };
+
+					// This step is also needed to resolve any invalid frame number rounding which may have occurred
+					// during the parsing of the frame log (as it implicitly does a conversion to/from frame number).
 					AudioTimecode = FTimecode::FromTimespan(Timecode.ToTimespan(OriginalFrameRate), TargetFrameRate, true);
 					AudioTimecodeRate = TargetFrameRate;
 					bAudioTimecodeFound = true;
