Create directories for third party encoder output path

diff --git a/Engine/Plugins/VirtualProduction/CaptureManager/CaptureManagerApp/Source/CaptureDataConverter/Private/Nodes/ThirdPartyEncoder/CaptureConvertVideoDataThirdParty.cpp b/Engine/Plugins/VirtualProduction/CaptureManager/CaptureManagerApp/Source/CaptureDataConverter/Private/Nodes/ThirdPartyEncoder/CaptureConvertVideoDataThirdParty.cpp
--- a/Engine/Plugins/VirtualProduction/CaptureManager/CaptureManagerApp/Source/CaptureDataConverter/Private/Nodes/ThirdPartyEncoder/CaptureConvertVideoDataThirdParty.cpp
+++ b/Engine/Plugins/VirtualProduction/CaptureManager/CaptureManagerApp/Source/CaptureDataConverter/Private/Nodes/ThirdPartyEncoder/CaptureConvertVideoDataThirdParty.cpp
@@ -161,6 +161,17 @@ FCaptureConvertVideoDataThirdParty::FResult FCaptureConvertVideoDataThirdParty::
 
 	const FString VideoFilePath = FPaths::ConvertRelativePathToFull(Params.TakeOriginDirectory, Video.Path);
 	FString DestinationDirectory = OutputDirectory / Video.Name;
+	
+	if (!FPaths::DirectoryExists(DestinationDirectory))
+	{
+		bool bSuccess = IFileManager::Get().MakeDirectory(*DestinationDirectory, true);
+		if (!bSuccess)
+		{
+			UE_LOG(LogCaptureConvertVideoDataThirdParty, Error, TEXT("Failed to create directory: %s"), *DestinationDirectory);
+			FText Message = LOCTEXT("CaptureConvertVideoNodeTP_FailedToCreateDirectory", "Failed to create encoder output directory: " + DestinationDirectory);
+			return MakeError(MoveTemp(Message));
+		}
+	}
 
 	const FString ImageFileName = FString::Format(TEXT("{0}_%06d.{1}"), { VideoParams.ImageFileName, VideoParams.Format });
 	const FString ImageFilePath = FPaths::Combine(DestinationDirectory, ImageFileName);
