[CaptureManager] Fixed crash on close in GetCaptureDevice

diff --git a/Engine/Plugins/VirtualProduction/CaptureManager/CaptureManagerApp/Source/CaptureManagerEditor/Private/UI/CaptureManagerPanelController.cpp b/Engine/Plugins/VirtualProduction/CaptureManager/CaptureManagerApp/Source/CaptureManagerEditor/Private/UI/CaptureManagerPanelController.cpp
--- a/Engine/Plugins/VirtualProduction/CaptureManager/CaptureManagerApp/Source/CaptureManagerEditor/Private/UI/CaptureManagerPanelController.cpp
+++ b/Engine/Plugins/VirtualProduction/CaptureManager/CaptureManagerApp/Source/CaptureManagerEditor/Private/UI/CaptureManagerPanelController.cpp
@@ -153,13 +153,17 @@ TSharedPtr<STakesView> FCaptureManagerPanelController::GetTakesView() const
 TObjectPtr<ULiveLinkDevice> FCaptureManagerPanelController::GetCaptureDevice(FGuid InDeviceId) const
 {
 	ULiveLinkDeviceSubsystem* Subsystem = GEngine->GetEngineSubsystem<ULiveLinkDeviceSubsystem>();
-	const TObjectPtr<ULiveLinkDevice>* DevicePtr = Subsystem->GetDeviceMap().Find(InDeviceId);
 
-	if (DevicePtr)
+	if (IsValid(Subsystem))
 	{
-		TObjectPtr<ULiveLinkDevice> Device = *DevicePtr;
-		check(Device->Implements<ULiveLinkDeviceCapability_Ingest>());
-		return Device;
+		const TObjectPtr<ULiveLinkDevice>* DevicePtr = Subsystem->GetDeviceMap().Find(InDeviceId);
+
+		if (DevicePtr)
+		{
+			TObjectPtr<ULiveLinkDevice> Device = *DevicePtr;
+			check(Device->Implements<ULiveLinkDeviceCapability_Ingest>());
+			return Device;
+		}
 	}
 
 	return nullptr;
@@ -170,17 +174,19 @@ TArray<TObjectPtr<ULiveLinkDevice>> FCaptureManagerPanelController::GetCaptureDe
 	ULiveLinkDeviceSubsystem* Subsystem = GEngine->GetEngineSubsystem<ULiveLinkDeviceSubsystem>();
 
 	TArray<TObjectPtr<ULiveLinkDevice>> CaptureDevices;
-
-	for (const TPair<FGuid, TObjectPtr<ULiveLinkDevice>>& DevicePair : Subsystem->GetDeviceMap())
+	if (IsValid(Subsystem))
 	{
-		TObjectPtr<ULiveLinkDevice> Device = DevicePair.Value;
-
-		if (!Device->Implements<ULiveLinkDeviceCapability_Ingest>())
+		for (const TPair<FGuid, TObjectPtr<ULiveLinkDevice>>& DevicePair : Subsystem->GetDeviceMap())
 		{
-			continue;
+			TObjectPtr<ULiveLinkDevice> Device = DevicePair.Value;
+
+			if (!Device->Implements<ULiveLinkDeviceCapability_Ingest>())
+			{
+				continue;
+			}
+
+			CaptureDevices.Add(MoveTemp(Device));
 		}
-		
-		CaptureDevices.Add(MoveTemp(Device));
 	}
 
 	return CaptureDevices;
diff --git a/Engine/Plugins/VirtualProduction/CaptureManager/CaptureManagerApp/Source/CaptureManagerEditor/Private/UI/TakesView.cpp b/Engine/Plugins/VirtualProduction/CaptureManager/CaptureManagerApp/Source/CaptureManagerEditor/Private/UI/TakesView.cpp
--- a/Engine/Plugins/VirtualProduction/CaptureManager/CaptureManagerApp/Source/CaptureManagerEditor/Private/UI/TakesView.cpp
+++ b/Engine/Plugins/VirtualProduction/CaptureManager/CaptureManagerApp/Source/CaptureManagerEditor/Private/UI/TakesView.cpp
@@ -725,6 +725,11 @@ void FTakesView::CaptureDeviceAdded(ULiveLinkDevice* InDevice)
 					TArray<FTakeId> TakeIds = { Event->TakeId };
 					TArray<TUniqueObjectPtr<UTakeVirtualAsset>> NewTakeObject = MakeTakeObjects(CaptureDeviceId, { TakeIds });
 
+					if (NewTakeObject.IsEmpty())
+					{
+						return;
+					}
+
 					if (TArray<TUniqueObjectPtr<UTakeVirtualAsset>>* TakeObjects = CaptureDeviceTakesMap.Find(CaptureDeviceId))
 					{
 						TakeObjects->Append(MoveTemp(NewTakeObject));
