[CaptureManager] Fix crashes in BaseIngestLiveLinkDevice

diff --git a/Engine/Plugins/VirtualProduction/CaptureManager/CaptureManagerApp/Source/IngestLiveLinkDevice/Private/BaseIngestLiveLinkDevice.cpp b/Engine/Plugins/VirtualProduction/CaptureManager/CaptureManagerApp/Source/IngestLiveLinkDevice/Private/BaseIngestLiveLinkDevice.cpp
--- a/Engine/Plugins/VirtualProduction/CaptureManager/CaptureManagerApp/Source/IngestLiveLinkDevice/Private/BaseIngestLiveLinkDevice.cpp
+++ b/Engine/Plugins/VirtualProduction/CaptureManager/CaptureManagerApp/Source/IngestLiveLinkDevice/Private/BaseIngestLiveLinkDevice.cpp
@@ -121,8 +121,8 @@ void UBaseIngestLiveLinkDevice::FImpl::IngestTake(const UIngestCapability_Proces
 
 	if (TakeFullPath.IsEmpty())
 	{
-		check(false);
 		FString Message = FString::Printf(TEXT("Failed to lookup take directory for take id %d"), TakeId);
+		ensureMsgf(false, TEXT("%s"), *Message);
 		Owner->ExecuteProcessFinishedReporter(InProcessHandle, MakeError(FIngestCapability_Error::InvalidArgument, MoveTemp(Message)));
 		return;
 	}
@@ -130,8 +130,8 @@ void UBaseIngestLiveLinkDevice::FImpl::IngestTake(const UIngestCapability_Proces
 	TOptional<FTakeMetadata> MaybeTakeMetadata = Owner->GetTakeMetadata(TakeId);
 	if (!MaybeTakeMetadata.IsSet())
 	{
-		check(false);
 		FString Message = FString::Printf(TEXT("Failed to lookup take metadata for take id %d (TakePath=%s)"), TakeId, *TakeFullPath);
+		ensureMsgf(false, TEXT("%s"), *Message);
 		Owner->ExecuteProcessFinishedReporter(InProcessHandle, MakeError(FIngestCapability_Error::InvalidArgument, MoveTemp(Message)));
 		return;
 	}
@@ -259,8 +259,7 @@ void UBaseIngestLiveLinkDevice::FImpl::IngestTake(const UIngestCapability_Proces
 		}
 		else
 		{
-			check(false);
-			UE_LOG(LogIngestLiveLinkDevice, Error, TEXT("Failed to add take upload task to queue"));
+			ensureMsgf(false, TEXT("Failed to lookup take directory for take id %d"), TakeId);
 		}
 	}
 	else
