[CaptureManager] Fixing a broken promise

diff --git a/Engine/Plugins/VirtualProduction/CaptureManager/CaptureManagerApp/Source/IngestLiveLinkDevice/Private/Utils/CaptureExtractTimecode.cpp b/Engine/Plugins/VirtualProduction/CaptureManager/CaptureManagerApp/Source/IngestLiveLinkDevice/Private/Utils/CaptureExtractTimecode.cpp
--- a/Engine/Plugins/VirtualProduction/CaptureManager/CaptureManagerApp/Source/IngestLiveLinkDevice/Private/Utils/CaptureExtractTimecode.cpp
+++ b/Engine/Plugins/VirtualProduction/CaptureManager/CaptureManagerApp/Source/IngestLiveLinkDevice/Private/Utils/CaptureExtractTimecode.cpp
@@ -64,34 +64,45 @@ UE::CaptureManager::FTimecodeAndRateResult FCaptureExtractVideoTimecode::Extract
 {
 	using namespace UE::CaptureManager;
 
-	using EventSinkResult = TValueOrError<void, FString>;
-	TPromise<EventSinkResult> EventSinkPromise;
+	using FEventSinkResult = TValueOrError<void, ECaptureExtractTimecodeError>;
+	TPromise<FEventSinkResult> EventSinkPromise;
 
 	class FMediaEventSink : public IMediaEventSink
 	{
 	public:
 
-		FMediaEventSink(TPromise<EventSinkResult>& InPromise)
+		FMediaEventSink(TPromise<FEventSinkResult>& InPromise)
 			: Promise(InPromise)
+			, bReleased(false)
 		{
 		}
 
 		virtual void ReceiveMediaEvent(EMediaEvent Event) override
 		{
+			if (bReleased)
+			{
+				return;
+			}
+
 			if (Event == EMediaEvent::MediaOpened)
 			{
 				Promise.SetValue(MakeValue());
 			}
 			else if (Event == EMediaEvent::MediaOpenFailed)
 			{
-				FString Message = TEXT("Failed to open the media file.");
-				Promise.SetValue(MakeError(MoveTemp(Message)));
+				Promise.SetValue(MakeError(ECaptureExtractTimecodeError::UnableToOpenMedia));
 			}
 		}
 
+		void Release()
+		{
+			bReleased = true;
+		}
+
 	private:
 
-		TPromise<EventSinkResult>& Promise;
+		TPromise<FEventSinkResult>& Promise;
+		std::atomic_bool bReleased;
 
 	} MediaEventSink(EventSinkPromise);
 
@@ -100,7 +111,10 @@ UE::CaptureManager::FTimecodeAndRateResult FCaptureExtractVideoTimecode::Extract
 	TSharedPtr<IMediaPlayer, ESPMode::ThreadSafe> MediaPlayer = ElectraModule.CreatePlayer(MediaEventSink);
 	if (!MediaPlayer)
 	{
-		return MakeError(ECaptureExtractTimecodeError::InternalError);
+		// Broken promises are considered errors
+		UE::CaptureManager::FTimecodeAndRateResult Result = MakeError(ECaptureExtractTimecodeError::InternalError);
+		EventSinkPromise.SetValue(MakeError(Result.GetError()));
+		return Result;
 	}
 
 	UDesiredPlayerMediaSource* MediaSource = NewObject<UDesiredPlayerMediaSource>();
@@ -119,7 +133,7 @@ UE::CaptureManager::FTimecodeAndRateResult FCaptureExtractVideoTimecode::Extract
 
 	FDateTime EventWaitStart = FDateTime::Now();
 
-	TFuture<EventSinkResult> EventSinkFuture = EventSinkPromise.GetFuture();
+	TFuture<FEventSinkResult> EventSinkFuture = EventSinkPromise.GetFuture();
 	
 	while ((FDateTime::Now() - EventWaitStart).GetSeconds() < TimeoutPeriod)
 	{
@@ -127,10 +141,10 @@ UE::CaptureManager::FTimecodeAndRateResult FCaptureExtractVideoTimecode::Extract
 
 		if (EventSinkFuture.WaitFor(FTimespan::FromMilliseconds(100.0f))) // Wait 100 milliseconds before invoking TickInput again
 		{
-			EventSinkResult Result = EventSinkFuture.Get();
+			FEventSinkResult Result = EventSinkFuture.Get();
 			if (Result.HasError())
 			{
-				return MakeError(ECaptureExtractTimecodeError::UnableToOpenMedia);
+				return MakeError(Result.StealError());
 			}
 
 			const FVariant Timecode = MediaPlayer->GetMediaInfo(UMediaPlayer::MediaInfoNameStartTimecodeValue.Resolve());
@@ -149,9 +163,21 @@ UE::CaptureManager::FTimecodeAndRateResult FCaptureExtractVideoTimecode::Extract
 		}
 	}
 
+	// Making sure that no additional events will arrive
+	MediaEventSink.Release();
+	MediaPlayer->Close();
+
 	if (TimecodeString.IsEmpty())
 	{
-		return MakeError(ECaptureExtractTimecodeError::TimecodeNotFound);
+		UE::CaptureManager::FTimecodeAndRateResult Result = MakeError(ECaptureExtractTimecodeError::TimecodeNotFound);
+
+		if (!EventSinkFuture.IsReady())
+		{
+			// Broken promises are considered errors
+			EventSinkPromise.SetValue(MakeError(Result.GetError()));
+		}
+
+		return Result;
 	}
 	
 	const TOptional<FTimecode> Timecode = FTimecode::ParseTimecode(*TimecodeString);
