FText::AsLocalizable_Advanced_LocText

LOCTEXT macros pass three wide string literals, two of which then construct a temp FTextKey at every call site.
Add a new method whose arguments match the macro so the FTextKeys are constructed inside.

Separate method name to avoid potential overload resolution troubles.

diff --git a/Engine/Plugins/VirtualProduction/CaptureManager/CaptureManagerApp/Source/CaptureDataConverter/Private/Nodes/ThirdPartyEncoder/CaptureConvertVideoDataThirdParty.cpp b/Engine/Plugins/VirtualProduction/CaptureManager/CaptureManagerApp/Source/CaptureDataConverter/Private/Nodes/ThirdPartyEncoder/CaptureConvertVideoDataThirdParty.cpp
--- a/Engine/Plugins/VirtualProduction/CaptureManager/CaptureManagerApp/Source/CaptureDataConverter/Private/Nodes/ThirdPartyEncoder/CaptureConvertVideoDataThirdParty.cpp
+++ b/Engine/Plugins/VirtualProduction/CaptureManager/CaptureManagerApp/Source/CaptureDataConverter/Private/Nodes/ThirdPartyEncoder/CaptureConvertVideoDataThirdParty.cpp
@@ -168,7 +168,11 @@ FCaptureConvertVideoDataThirdParty::FResult FCaptureConvertVideoDataThirdParty::
 		if (!bSuccess)
 		{
 			UE_LOG(LogCaptureConvertVideoDataThirdParty, Error, TEXT("Failed to create directory: %s"), *DestinationDirectory);
-			FText Message = LOCTEXT("CaptureConvertVideoNodeTP_FailedToCreateDirectory", "Failed to create encoder output directory: " + DestinationDirectory);
+			FText Message = FText::Format(
+				LOCTEXT("CaptureConvertVideoNodeTP_FailedToCreateDirectory", "Failed to create encoder output directory: {0}"),
+				FText::FromString(DestinationDirectory)
+			);
+
 			return MakeError(MoveTemp(Message));
 		}
 	}
@@ -363,4 +367,4 @@ FCaptureConvertVideoDataThirdParty::FResult FCaptureConvertVideoDataThirdParty::
 	return MakeValue();
 }
 
-#undef LOCTEXT_NAMESPACE
\ No newline at end of file
+#undef LOCTEXT_NAMESPACE
diff --git a/Engine/Source/Runtime/Core/Private/Internationalization/Text.cpp b/Engine/Source/Runtime/Core/Private/Internationalization/Text.cpp
--- a/Engine/Source/Runtime/Core/Private/Internationalization/Text.cpp
+++ b/Engine/Source/Runtime/Core/Private/Internationalization/Text.cpp
@@ -1242,6 +1242,11 @@ FText FText::AsLocalizable_Advanced(const FTextKey& Namespace, const FTextKey& K
 		: FTextCache::Get().FindOrCache(MoveTemp(String), FTextId(Namespace, Key));
 }
 
+FText FText::AsLocalizable_Advanced_LocText(const TCHAR* Namespace, const TCHAR* Key, const TCHAR* String)
+{
+	return FText::AsLocalizable_Advanced(Namespace, Key, FStringView(String));
+}
+
 bool FText::FindTextInLiveTable_Advanced(const FTextKey& Namespace, const FTextKey& Key, FText& OutText, const FString* const SourceString)
 {
 	FTextConstDisplayStringPtr FoundString = FTextLocalizationManager::Get().FindDisplayString(Namespace, Key, SourceString);
diff --git a/Engine/Source/Runtime/Core/Public/Internationalization/Internationalization.h b/Engine/Source/Runtime/Core/Public/Internationalization/Internationalization.h
--- a/Engine/Source/Runtime/Core/Public/Internationalization/Internationalization.h
+++ b/Engine/Source/Runtime/Core/Public/Internationalization/Internationalization.h
@@ -292,12 +292,12 @@ namespace UE::Private::LocGen
  * Creates an FText. All parameters must be string literals. All literals will be passed through the localization system.
  * The global LOCTEXT_NAMESPACE macro must be first set to a string literal to specify this localization key's namespace.
  */
-#define LOCTEXT(InKey, InTextLiteral) FText::AsLocalizable_Advanced(TEXT(LOCTEXT_NAMESPACE), TEXT(InKey), TEXT(InTextLiteral))
+#define LOCTEXT(InKey, InTextLiteral) FText::AsLocalizable_Advanced_LocText(TEXT(LOCTEXT_NAMESPACE), TEXT(InKey), TEXT(InTextLiteral))
 
 /**
  * Creates an FText. All parameters must be string literals. All literals will be passed through the localization system.
  */
-#define NSLOCTEXT(InNamespace, InKey, InTextLiteral) FText::AsLocalizable_Advanced(TEXT(InNamespace), TEXT(InKey), TEXT(InTextLiteral))
+#define NSLOCTEXT(InNamespace, InKey, InTextLiteral) FText::AsLocalizable_Advanced_LocText(TEXT(InNamespace), TEXT(InKey), TEXT(InTextLiteral))
 
 /**
  * Creates a culture invariant FText from the given string literal.
diff --git a/Engine/Source/Runtime/Core/Public/Internationalization/Text.h b/Engine/Source/Runtime/Core/Public/Internationalization/Text.h
--- a/Engine/Source/Runtime/Core/Public/Internationalization/Text.h
+++ b/Engine/Source/Runtime/Core/Public/Internationalization/Text.h
@@ -535,6 +535,7 @@ public:
 	[[nodiscard]] static CORE_API FText AsLocalizable_Advanced(const FTextKey& Namespace, const FTextKey& Key, const TCHAR* String);
 	[[nodiscard]] static CORE_API FText AsLocalizable_Advanced(const FTextKey& Namespace, const FTextKey& Key, FStringView String);
 	[[nodiscard]] static CORE_API FText AsLocalizable_Advanced(const FTextKey& Namespace, const FTextKey& Key, FString&& String);
+	[[nodiscard]] static CORE_API FText AsLocalizable_Advanced_LocText(const TCHAR* Namespace, const TCHAR* Key, const TCHAR* String); // For use by LOCTEXT to avoid constructing temp FTextKeys in every caller
 
 	/**
 	 * === !! This is an ADVANCED function. USE WITH CAUTION !! ===
