[CaptureManager] Fixed inaccurate comment

diff --git a/Engine/Plugins/VirtualProduction/CaptureManager/CaptureManagerApp/Source/CaptureManagerMediaRW/Private/Readers/WindowsMediaReader.cpp b/Engine/Plugins/VirtualProduction/CaptureManager/CaptureManagerApp/Source/CaptureManagerMediaRW/Private/Readers/WindowsMediaReader.cpp
--- a/Engine/Plugins/VirtualProduction/CaptureManager/CaptureManagerApp/Source/CaptureManagerMediaRW/Private/Readers/WindowsMediaReader.cpp
+++ b/Engine/Plugins/VirtualProduction/CaptureManager/CaptureManagerApp/Source/CaptureManagerMediaRW/Private/Readers/WindowsMediaReader.cpp
@@ -57,7 +57,7 @@ TOptional<FText> FWindowsAudioReader::Open(const FString& InFileName)
 	int64 DurationTicks;	// 10,000,000 == 1.0e7 units per second. Each unit corresponds to 100 nanoseconds.
 
 	Result = AudioReader->GetPresentationAttribute(MF_SOURCE_READER_MEDIASOURCE, MF_PD_DURATION, &DurationVar);
-	WINAR_CHECK_AND_RETURN_ERROR_MESSAGE(Result, LOCTEXT("OpenAudio_FailedToRetrieveDuration", "Failed to retrieve the video file duration"));
+	WINAR_CHECK_AND_RETURN_ERROR_MESSAGE(Result, LOCTEXT("OpenAudio_FailedToRetrieveDuration", "Failed to retrieve the audio file duration"));
 
 	Result = ::PropVariantToInt64(DurationVar, &DurationTicks);
 	PropVariantClear(&DurationVar);
