[CaptureManager] Fixed build breaks in ExampleNetworkIngestDevice

diff --git a/Engine/Plugins/VirtualProduction/CaptureManager/CaptureManagerApp/Source/ExampleLiveLinkDevices/Private/NetworkIngest/ExampleNetworkIngestDevice.cpp b/Engine/Plugins/VirtualProduction/CaptureManager/CaptureManagerApp/Source/ExampleLiveLinkDevices/Private/NetworkIngest/ExampleNetworkIngestDevice.cpp
--- a/Engine/Plugins/VirtualProduction/CaptureManager/CaptureManagerApp/Source/ExampleLiveLinkDevices/Private/NetworkIngest/ExampleNetworkIngestDevice.cpp
+++ b/Engine/Plugins/VirtualProduction/CaptureManager/CaptureManagerApp/Source/ExampleLiveLinkDevices/Private/NetworkIngest/ExampleNetworkIngestDevice.cpp
@@ -6,10 +6,10 @@
 
 const UExampleNetworkIngestDeviceSettings* UExampleNetworkIngestDevice::GetSettings() const
 {
-	return CastChecked<UExampleNetworkIngestDeviceSettings>(Settings);
+	return GetDeviceSettings<UExampleNetworkIngestDeviceSettings>();
 }
 
-TSubclassOf<ULiveLinkDeviceSettings> UExampleNetworkIngestDevice::StaticGetSettingsClass() const
+TSubclassOf<ULiveLinkDeviceSettings> UExampleNetworkIngestDevice::GetSettingsClass() const
 {
 	return UExampleNetworkIngestDeviceSettings::StaticClass();
 }
@@ -29,11 +29,11 @@ FText UExampleNetworkIngestDevice::GetHealthText() const
 	return FText::FromString("Example Health");
 }
 
-void UExampleNetworkIngestDevice::OnDeviceAdded(const FGuid InDeviceGuid, ULiveLinkDeviceSettings* InSettings)
+void UExampleNetworkIngestDevice::OnDeviceAdded()
 {
 	// Respond to device added event here
 
-	Super::OnDeviceAdded(InDeviceGuid, InSettings);
+	Super::OnDeviceAdded();
 }
 
 void UExampleNetworkIngestDevice::OnDeviceRemoved()
@@ -80,7 +80,7 @@ void UExampleNetworkIngestDevice::UpdateTakeList_Implementation(UIngestCapabilit
 	// See ULiveLinkFaceDevice::UpdateTakeList_Implementation
 }
 
-void UExampleNetworkIngestDevice::RunIngestTake(const UIngestCapability_ProcessHandle* InProcessHandle, const UIngestCapability_Options* InIngestOptions)
+void UExampleNetworkIngestDevice::RunConvertAndUploadTake(const UIngestCapability_ProcessHandle* InProcessHandle, const UIngestCapability_Options* InIngestOptions)
 {
 	// Get the take corresponding to the InTakeId from your device (i.e. Download the take from the device)
 	//     See ULiveLinkFaceDevice::RunIngestTake
@@ -91,9 +91,9 @@ void UExampleNetworkIngestDevice::RunIngestTake(const UIngestCapability_ProcessH
 
 	using namespace UE::CaptureManager;
 	TSharedPtr<FTaskProgress> TaskProgress = MakeShared<FTaskProgress>
-		(NumberOfTasks, FTaskProgress::FProgressReporter::CreateLambda([this, TakeId = InProcessHandle->TakeId](double InProgress)
+		(NumberOfTasks, FTaskProgress::FProgressReporter::CreateLambda([this, InProcessHandle](double InProgress)
 			{
-				ExecuteProcessProgressReporter(TakeId, InProgress);
+				ExecuteProcessProgressReporter(InProcessHandle, InProgress);
 			}));
 
 	// Note: IngestTake uses GetFullTakePath() so make sure you have implemented it
@@ -105,7 +105,7 @@ void UExampleNetworkIngestDevice::CancelIngestProcess_Implementation(const UInge
 {
 	// Stop your device from sending data here
 
-	Super::CancelIngest(InProcessHandle->TakeId);
+	Super::CancelIngest(InProcessHandle->GetTakeId());
 }
 
 ELiveLinkDeviceConnectionStatus UExampleNetworkIngestDevice::GetConnectionStatus_Implementation() const
diff --git a/Engine/Plugins/VirtualProduction/CaptureManager/CaptureManagerApp/Source/ExampleLiveLinkDevices/Private/NetworkIngest/ExampleNetworkIngestDevice.h b/Engine/Plugins/VirtualProduction/CaptureManager/CaptureManagerApp/Source/ExampleLiveLinkDevices/Private/NetworkIngest/ExampleNetworkIngestDevice.h
--- a/Engine/Plugins/VirtualProduction/CaptureManager/CaptureManagerApp/Source/ExampleLiveLinkDevices/Private/NetworkIngest/ExampleNetworkIngestDevice.h
+++ b/Engine/Plugins/VirtualProduction/CaptureManager/CaptureManagerApp/Source/ExampleLiveLinkDevices/Private/NetworkIngest/ExampleNetworkIngestDevice.h
@@ -42,12 +42,12 @@ public:
 	const UExampleNetworkIngestDeviceSettings* GetSettings() const;
 
 	//~ Begin ULiveLinkDevice interface
-	virtual TSubclassOf<ULiveLinkDeviceSettings> StaticGetSettingsClass() const override;
+	virtual TSubclassOf<ULiveLinkDeviceSettings> GetSettingsClass() const override;
 	virtual FText GetDisplayName() const override;
 	virtual EDeviceHealth GetDeviceHealth() const override;
 	virtual FText GetHealthText() const override;
 
-	virtual void OnDeviceAdded(const FGuid InDeviceGuid, ULiveLinkDeviceSettings* InSettings) override;
+	virtual void OnDeviceAdded() override;
 	virtual void OnDeviceRemoved() override;
 	virtual void OnSettingChanged(const FPropertyChangedEvent& InPropertyChangedEvent) override;
 	//~ End ULiveLinkDevice interface
@@ -57,7 +57,7 @@ private:
 
 	//~ Begin ULiveLinkDeviceCapability_Ingest interface
 	virtual void UpdateTakeList_Implementation(UIngestCapability_UpdateTakeListCallback* InCallback) override;
-	virtual void RunIngestTake(const UIngestCapability_ProcessHandle* InProcessHandle,
+	virtual void RunConvertAndUploadTake(const UIngestCapability_ProcessHandle* InProcessHandle,
 							   const UIngestCapability_Options* InIngestOptions) override;
 	virtual void CancelIngestProcess_Implementation(const UIngestCapability_ProcessHandle* InProcessHandle) override;
 	//~ End ULiveLinkDeviceCapability_Ingest interface
