[CaptureManager] Add device token to default working directory and download directory

diff --git a/Engine/Plugins/VirtualProduction/CaptureManager/CaptureManagerApp/Source/CaptureManagerSettings/Private/Settings/CaptureManagerSettings.cpp b/Engine/Plugins/VirtualProduction/CaptureManager/CaptureManagerApp/Source/CaptureManagerSettings/Private/Settings/CaptureManagerSettings.cpp
--- a/Engine/Plugins/VirtualProduction/CaptureManager/CaptureManagerApp/Source/CaptureManagerSettings/Private/Settings/CaptureManagerSettings.cpp
+++ b/Engine/Plugins/VirtualProduction/CaptureManager/CaptureManagerApp/Source/CaptureManagerSettings/Private/Settings/CaptureManagerSettings.cpp
@@ -9,8 +9,8 @@
 
 namespace UE::CaptureManager::Private
 {
-constexpr FStringView DefaultWorkingDirectory = TEXTVIEW("CaptureManager/Working/{yyyy}{mm}{dd}");
-constexpr FStringView DefaultDownloadDirectory = TEXTVIEW("CaptureManager/Download/{yyyy}{mm}{dd}");
+constexpr FStringView DefaultWorkingDirectory = TEXTVIEW("CaptureManager/Working/{yyyy}{mm}{dd}/{device}");
+constexpr FStringView DefaultDownloadDirectory = TEXTVIEW("CaptureManager/Download/{yyyy}{mm}{dd}/{device}");
 }
 
 UCaptureManagerSettings::UCaptureManagerSettings(const FObjectInitializer& InObjectInitializer) :
