Supporting sample rate of 16kHz for audio stream

diff --git a/Engine/Plugins/VirtualProduction/CaptureManager/CaptureManagerApp/Source/CaptureManagerMediaRW/Private/MediaSample.cpp b/Engine/Plugins/VirtualProduction/CaptureManager/CaptureManagerApp/Source/CaptureManagerMediaRW/Private/MediaSample.cpp
--- a/Engine/Plugins/VirtualProduction/CaptureManager/CaptureManagerApp/Source/CaptureManagerMediaRW/Private/MediaSample.cpp
+++ b/Engine/Plugins/VirtualProduction/CaptureManager/CaptureManagerApp/Source/CaptureManagerMediaRW/Private/MediaSample.cpp
@@ -27,6 +27,10 @@ int32 ConvertSampleRate(ESampleRate InSampleRate)
 {
 	switch (InSampleRate)
 	{
+		case UE::CaptureManager::ESampleRate::SR_8000Hz:
+			return 8000;
+		case UE::CaptureManager::ESampleRate::SR_16000Hz:
+			return 16000;
 		case UE::CaptureManager::ESampleRate::SR_48000Hz:
 			return 48000;
 		case UE::CaptureManager::ESampleRate::SR_88200Hz:
@@ -45,6 +49,10 @@ ESampleRate ConvertSampleRate(int32 InSampleRate)
 {
 	switch (InSampleRate)
 	{
+		case 8000:
+			return UE::CaptureManager::ESampleRate::SR_8000Hz;
+		case 16000:
+			return UE::CaptureManager::ESampleRate::SR_16000Hz;
 		case 48000:
 			return UE::CaptureManager::ESampleRate::SR_48000Hz;
 		case 88200:
diff --git a/Engine/Plugins/VirtualProduction/CaptureManager/CaptureManagerApp/Source/CaptureManagerMediaRW/Public/MediaSample.h b/Engine/Plugins/VirtualProduction/CaptureManager/CaptureManagerApp/Source/CaptureManagerMediaRW/Public/MediaSample.h
--- a/Engine/Plugins/VirtualProduction/CaptureManager/CaptureManagerApp/Source/CaptureManagerMediaRW/Public/MediaSample.h
+++ b/Engine/Plugins/VirtualProduction/CaptureManager/CaptureManagerApp/Source/CaptureManagerMediaRW/Public/MediaSample.h
@@ -10,7 +10,9 @@ namespace UE::CaptureManager
 
 enum class ESampleRate
 {
-	SR_44100Hz = 0,
+	SR_8000Hz = 0,
+	SR_16000Hz,
+	SR_44100Hz,
 	SR_48000Hz,
 	SR_88200Hz,
 	SR_96000Hz,
