Fixing Python issues:

* Not removing the shared reference when the disconnect happens
* Initiating a disconnect instead of stopping the endpoint manager

diff --git a/Engine/Plugins/VirtualProduction/CaptureManager/CaptureManagerApp/Source/CaptureManagerUnrealEndpoint/Private/CaptureCommunicationManager.cpp b/Engine/Plugins/VirtualProduction/CaptureManager/CaptureManagerApp/Source/CaptureManagerUnrealEndpoint/Private/CaptureCommunicationManager.cpp
--- a/Engine/Plugins/VirtualProduction/CaptureManager/CaptureManagerApp/Source/CaptureManagerUnrealEndpoint/Private/CaptureCommunicationManager.cpp
+++ b/Engine/Plugins/VirtualProduction/CaptureManager/CaptureManagerApp/Source/CaptureManagerUnrealEndpoint/Private/CaptureCommunicationManager.cpp
@@ -55,13 +55,12 @@ FCommunicationManager::FCommunicationManager()
 
 FCommunicationManager::~FCommunicationManager()
 {
+	Disconnect();
+
 	Messenger->SetDisconnectHandler(nullptr);
 	Messenger->SetUploadCallbacks(nullptr, nullptr);
 
 	Messenger = nullptr;
-
-	FScopeLock Lock(&Mutex);
-	Uploaders.Empty();
 }
 
 void FCommunicationManager::Connect(FMessageAddress InAddress, FString InIpAddress, uint16 InPort)
@@ -99,14 +98,6 @@ void FCommunicationManager::Disconnect()
 		return;
 	}
 
-	Messenger->SetDisconnectHandler(
-		FConnectStarter::FDisconnectHandler::CreateLambda(
-		[this]()
-		{
-			ConnectionStateChangedDelegate.ExecuteIfBound(Messenger->GetAddress(), false);
-		}
-	)
-	);
 	Messenger->Disconnect();
 
 	FScopeLock Lock(&Mutex);
diff --git a/Engine/Plugins/VirtualProduction/CaptureManager/CaptureManagerApp/Source/CaptureManagerUnrealEndpoint/Private/CaptureManagerUnrealEndpoint.cpp b/Engine/Plugins/VirtualProduction/CaptureManager/CaptureManagerApp/Source/CaptureManagerUnrealEndpoint/Private/CaptureManagerUnrealEndpoint.cpp
--- a/Engine/Plugins/VirtualProduction/CaptureManager/CaptureManagerApp/Source/CaptureManagerUnrealEndpoint/Private/CaptureManagerUnrealEndpoint.cpp
+++ b/Engine/Plugins/VirtualProduction/CaptureManager/CaptureManagerApp/Source/CaptureManagerUnrealEndpoint/Private/CaptureManagerUnrealEndpoint.cpp
@@ -95,15 +95,18 @@ FUnrealEndpoint::FImpl::~FImpl() = default;
 void FUnrealEndpoint::FImpl::SetupConnectionDelegate()
 {
 	CommunicationManager.ConnectionStateChanged().BindLambda(
-		[This = AsShared()](const FMessageAddress& InMessageAddress, bool bInIsConnected)
+		[WeakThis = AsWeak()](const FMessageAddress& InMessageAddress, bool bInIsConnected)
 		{
-			if (InMessageAddress == This->EndpointInfo.MessageAddress)
+			if (TSharedPtr<FImpl> This = WeakThis.Pin())
 			{
+				if (InMessageAddress == This->EndpointInfo.MessageAddress)
 				{
-					std::lock_guard<std::mutex> LockGuard(This->CondMtx);
-					This->bIsConnected = bInIsConnected;
+					{
+						std::lock_guard<std::mutex> LockGuard(This->CondMtx);
+						This->bIsConnected = bInIsConnected;
+					}
+					This->CondVar.notify_one();
 				}
-				This->CondVar.notify_one();
 			}
 		}
 	);
@@ -175,51 +178,57 @@ bool FUnrealEndpoint::FImpl::AddTakeUploadTask(TUniquePtr<FTakeUploadTask> InTak
 
 	// Note: Only the first call to GetUploader will register these callbacks
 	FUploader::FUploadStateUpdate UpdateCallback = FUploader::FUploadStateUpdate::CreateLambda(
-		[This = AsShared()](const FGuid& InTakeUploadID, const double InProgress)
+		[WeakThis = AsWeak()](const FGuid& InTakeUploadID, const double InProgress)
 		{
-			// Find the take upload task and update it
-			FScopeLock ScopeLock(&This->CriticalSection);
+			if (TSharedPtr<FImpl> This = WeakThis.Pin())
+			{
+				// Find the take upload task and update it
+				FScopeLock ScopeLock(&This->CriticalSection);
 
-			const int32 TaskIndex = This->TakeUploadTasks.IndexOfByPredicate(
-				[&InTakeUploadID](const TSharedRef<FTakeUploadTask>& InTask)
-				{
-					return InTask->GetTaskID() == InTakeUploadID;
-				}
-			);
+				const int32 TaskIndex = This->TakeUploadTasks.IndexOfByPredicate(
+					[&InTakeUploadID](const TSharedRef<FTakeUploadTask>& InTask)
+					{
+						return InTask->GetTaskID() == InTakeUploadID;
+					}
+				);
 
-			if (TaskIndex != INDEX_NONE)
-			{
-				// Grab a strong reference before releasing the lock
-				TSharedRef<FTakeUploadTask> ThisTask = This->TakeUploadTasks[TaskIndex];
-				ScopeLock.Unlock();
+				if (TaskIndex != INDEX_NONE)
+				{
+					// Grab a strong reference before releasing the lock
+					TSharedRef<FTakeUploadTask> ThisTask = This->TakeUploadTasks[TaskIndex];
+					ScopeLock.Unlock();
 
-				ThisTask->Progressed().ExecuteIfBound(InProgress);
+					ThisTask->Progressed().ExecuteIfBound(InProgress);
+				}
 			}
 		}
 	);
 
 	// Note: Only the first call to GetUploader will register these callbacks
 	FUploader::FUploadFinished UpdateFinishedCallback = FUploader::FUploadFinished::CreateLambda(
-		[This = AsShared()](const FGuid& InTakeUploadID, const FString& InMessage, int32 InCode)
+		[WeakThis = AsWeak()](const FGuid& InTakeUploadID, const FString& InMessage, int32 InCode)
 		{
-			FScopeLock ScopeLock(&This->CriticalSection);
+			if (TSharedPtr<FImpl> This = WeakThis.Pin())
+			{
+				FScopeLock ScopeLock(&This->CriticalSection);
 
-			// Find the take upload task and update it
-			const int32 TaskIndex = This->TakeUploadTasks.IndexOfByPredicate(
-				[&InTakeUploadID](const TSharedRef<FTakeUploadTask>& InTask)
-				{
-					return InTask->GetTaskID() == InTakeUploadID;
-				}
-			);
+				// Find the take upload task and update it
+				const int32 TaskIndex = This->TakeUploadTasks.IndexOfByPredicate(
+					[&InTakeUploadID](const TSharedRef<FTakeUploadTask>& InTask)
+					{
+						return InTask->GetTaskID() == InTakeUploadID;
+					}
+				);
 
-			if (TaskIndex != INDEX_NONE)
-			{
-				// Grab a strong reference before releasing the lock
-				TSharedRef<FTakeUploadTask> ThisTask = This->TakeUploadTasks[TaskIndex];
-				This->TakeUploadTasks.RemoveAt(TaskIndex);
-				ScopeLock.Unlock();
+				if (TaskIndex != INDEX_NONE)
+				{
+					// Grab a strong reference before releasing the lock
+					TSharedRef<FTakeUploadTask> ThisTask = This->TakeUploadTasks[TaskIndex];
+					This->TakeUploadTasks.RemoveAt(TaskIndex);
+					ScopeLock.Unlock();
 
-				ThisTask->Complete().ExecuteIfBound(InMessage, InCode);
+					ThisTask->Complete().ExecuteIfBound(InMessage, InCode);
+				}
 			}
 		}
 	);
diff --git a/Engine/Plugins/VirtualProduction/CaptureManager/CaptureManagerCore/Source/LiveLinkHubCaptureMessaging/Private/Features/ConnectStarter.cpp b/Engine/Plugins/VirtualProduction/CaptureManager/CaptureManagerCore/Source/LiveLinkHubCaptureMessaging/Private/Features/ConnectStarter.cpp
--- a/Engine/Plugins/VirtualProduction/CaptureManager/CaptureManagerCore/Source/LiveLinkHubCaptureMessaging/Private/Features/ConnectStarter.cpp
+++ b/Engine/Plugins/VirtualProduction/CaptureManager/CaptureManagerCore/Source/LiveLinkHubCaptureMessaging/Private/Features/ConnectStarter.cpp
@@ -120,6 +120,9 @@ void FConnectStarter::HandleConnectResponse(const FConnectResponse& InResponse,
 
 	if (InResponse.Status == EStatus::Ok)
 	{
+		// Remove previous keep alive entry (reset keep alive)
+		TimerManager->RemoveTimer(KeepAlive);
+
 		KeepAlive = 
 			TimerManager->AddTimer(
 				FTimerDelegate::CreateSP(this, &FConnectStarter::OnKeepAliveInterval), KeepAliveInterval, false, KeepAliveInterval);
