[CaptureManager] Add tooltips to ingest job status icon

diff --git a/Engine/Plugins/VirtualProduction/CaptureManager/CaptureManagerApp/Source/CaptureManagerEditor/Private/IngestManagement/IngestJob.cpp b/Engine/Plugins/VirtualProduction/CaptureManager/CaptureManagerApp/Source/CaptureManagerEditor/Private/IngestManagement/IngestJob.cpp
--- a/Engine/Plugins/VirtualProduction/CaptureManager/CaptureManagerApp/Source/CaptureManagerEditor/Private/IngestManagement/IngestJob.cpp
+++ b/Engine/Plugins/VirtualProduction/CaptureManager/CaptureManagerApp/Source/CaptureManagerEditor/Private/IngestManagement/IngestJob.cpp
@@ -41,8 +41,9 @@ const FTakeMetadata& FIngestJob::GetTakeMetadata() const
 	return TakeMetadata;
 }
 
-FIngestJob::EProcessingState FIngestJob::GetProcessingState() const
+FIngestJob::FProcessingState FIngestJob::GetProcessingState() const
 {
+	FScopeLock ScopeLock(&CriticalSection);
 	return ProcessingState;
 }
 
@@ -62,9 +63,10 @@ EIngestCapability_ProcessConfig FIngestJob::GetPipelineConfig() const
 	return PipelineConfig;
 }
 
-void FIngestJob::SetProcessingState(const FIngestJob::EProcessingState InProcessingState)
+void FIngestJob::SetProcessingState(FIngestJob::FProcessingState InProcessingState)
 {
-	ProcessingState = InProcessingState;
+	FScopeLock ScopeLock(&CriticalSection);
+	ProcessingState = MoveTemp(InProcessingState);
 }
 
 void FIngestJob::SetProgress(const double InProgress)
diff --git a/Engine/Plugins/VirtualProduction/CaptureManager/CaptureManagerApp/Source/CaptureManagerEditor/Private/IngestManagement/IngestJob.h b/Engine/Plugins/VirtualProduction/CaptureManager/CaptureManagerApp/Source/CaptureManagerEditor/Private/IngestManagement/IngestJob.h
--- a/Engine/Plugins/VirtualProduction/CaptureManager/CaptureManagerApp/Source/CaptureManagerEditor/Private/IngestManagement/IngestJob.h
+++ b/Engine/Plugins/VirtualProduction/CaptureManager/CaptureManagerApp/Source/CaptureManagerEditor/Private/IngestManagement/IngestJob.h
@@ -21,6 +21,12 @@ public:
 		Aborted = 8
 	};
 
+	struct FProcessingState
+	{
+		EProcessingState State;
+		FText Message;
+	};
+
 	struct FSettings
 	{
 		struct FVideoSettings
@@ -58,11 +64,11 @@ public:
 	const FTakeMetadata& GetTakeMetadata() const;
 	EIngestCapability_ProcessConfig GetPipelineConfig() const;
 	const FSettings& GetSettings() const;
-	EProcessingState GetProcessingState() const;
+	FProcessingState GetProcessingState() const;
 	double GetProgress() const;
 
 	void SetSettings(FSettings Settings);
-	void SetProcessingState(EProcessingState ProcessingState);
+	void SetProcessingState(FProcessingState ProcessingState);
 	void SetProgress(double Progress);
 
 private:
@@ -73,13 +79,15 @@ private:
 	const EIngestCapability_ProcessConfig PipelineConfig;
 
 	mutable FCriticalSection CriticalSection;
+
 	FSettings Settings;
 
 	static_assert(std::atomic<double>::is_always_lock_free);
 	std::atomic<double> Progress = 0.0;
 
-	static_assert(std::atomic<EProcessingState>::is_always_lock_free);
-	std::atomic<EProcessingState> ProcessingState = EProcessingState::Pending;
+	FProcessingState ProcessingState = FProcessingState{ EProcessingState::Pending, FText::FromString("Pending...") };
+
+	FText ProcessingStateMessage;
 };
 
 }
\ No newline at end of file
diff --git a/Engine/Plugins/VirtualProduction/CaptureManager/CaptureManagerApp/Source/CaptureManagerEditor/Private/IngestManagement/IngestJobExecutor.cpp b/Engine/Plugins/VirtualProduction/CaptureManager/CaptureManagerApp/Source/CaptureManagerEditor/Private/IngestManagement/IngestJobExecutor.cpp
--- a/Engine/Plugins/VirtualProduction/CaptureManager/CaptureManagerApp/Source/CaptureManagerEditor/Private/IngestManagement/IngestJobExecutor.cpp
+++ b/Engine/Plugins/VirtualProduction/CaptureManager/CaptureManagerApp/Source/CaptureManagerEditor/Private/IngestManagement/IngestJobExecutor.cpp
@@ -19,6 +19,8 @@
 #include "NamingTokensEngineSubsystem.h"
 #include "Settings/CaptureManagerTemplateTokens.h"
 
+#define LOCTEXT_NAMESPACE "IngestJobExecutor"
+
 DEFINE_LOG_CATEGORY_STATIC(LogCaptureManagerIngestJobExecutor, Log, All);
 
 namespace UE::CaptureManager
@@ -88,8 +90,9 @@ uint32 FIngestJobExecutor::Run()
 
 		if (Job)
 		{
-			check(Job->GetProcessingState() != FIngestJob::EProcessingState::Pending);
-			JobProcessingStateChanged.ExecuteIfBound(Job->GetGuid(), Job->GetProcessingState());
+			FIngestJob::EProcessingState State = Job->GetProcessingState().State;
+			check(State != FIngestJob::EProcessingState::Pending);
+			JobProcessingStateChanged.ExecuteIfBound(Job->GetGuid(), State);
 
 			FScopeLock Lock(&CriticalSection);
 
@@ -169,14 +172,15 @@ TUniquePtr<FIngestJobExecutor::FRunIngestContext> FIngestJobExecutor::RunIngest(
 {
 	if (!GEditor)
 	{
+		FText Message = LOCTEXT("IngestProcessStateMessage_GeditorUnavailable", "Ingest failed, GEditor unavailble");
 		UE_LOG(
 			LogCaptureManagerIngestJobExecutor,
 			Error,
-			TEXT("[%s] Ingest failed, GEditor unavailble"), *ExecutorName
+			TEXT("[%s] %s"), *ExecutorName, *Message.ToString()
 		);
 
 		// Set the processing state to anything other than Pending, to prevent other executors picking up the same bad job.
-		InJob.SetProcessingState(FIngestJob::EProcessingState::Aborted);
+		InJob.SetProcessingState(FIngestJob::FProcessingState{ FIngestJob::EProcessingState::Aborted, MoveTemp(Message) });
 		InPromise.SetValue();
 		return nullptr;
 	}
@@ -188,14 +192,15 @@ TUniquePtr<FIngestJobExecutor::FRunIngestContext> FIngestJobExecutor::RunIngest(
 
 	if (!DevicePtr)
 	{
+		FText Message = LOCTEXT("IngestProcessStateMessage_DeviceDoesNotExist", "Ingest failed, Capture Device does not exist");
 		UE_LOG(
 			LogCaptureManagerIngestJobExecutor,
 			Error,
-			TEXT("[%s] Ingest failed, Capture Device does not exist"), *ExecutorName
+			TEXT("[%s] %s"), *ExecutorName, *Message.ToString()
 		);
 
 		// Set the processing state to anything other than Pending, to prevent other executors picking up the same bad job.
-		InJob.SetProcessingState(FIngestJob::EProcessingState::Aborted);
+		InJob.SetProcessingState(FIngestJob::FProcessingState{ FIngestJob::EProcessingState::Aborted, MoveTemp(Message) });
 		InPromise.SetValue();
 		return nullptr;
 	}
@@ -260,14 +265,15 @@ TUniquePtr<FIngestJobExecutor::FRunIngestContext> FIngestJobExecutor::RunIngest(
 
 	if (!Device->Implements<ULiveLinkDeviceCapability_Ingest>())
 	{
+		FText Message = FText::Format(LOCTEXT("IngestProcessStateMessage_DeviceDoesNotSupportIngestCapability", "Ingest failed, Capture Device {0} does not support the ingest capability"), FText::FromString(Device->GetName()));
 		UE_LOG(
 			LogCaptureManagerIngestJobExecutor,
 			Error,
-			TEXT("[%s] Ingest failed, Capture Device does not support the ingest capability: %s"), *ExecutorName, *(Device->GetName())
+			TEXT("[%s] %s"), *ExecutorName, *Message.ToString()
 		);
 
 		// Set the processing state to anything other than Pending, to prevent other executors picking up the same bad job.
-		InJob.SetProcessingState(FIngestJob::EProcessingState::Aborted);
+		InJob.SetProcessingState(FIngestJob::FProcessingState{FIngestJob::EProcessingState::Aborted, MoveTemp(Message)});
 		InPromise.SetValue();
 		return nullptr;
 	}
@@ -276,26 +282,30 @@ TUniquePtr<FIngestJobExecutor::FRunIngestContext> FIngestJobExecutor::RunIngest(
 		[this, &InPromise, &InJob](const UIngestCapability_ProcessHandle* InProcessHandle, TValueOrError<void, FIngestCapability_Error> IngestError)
 		{
 			FIngestJob::EProcessingState ProcessingState;
+			FText Message;
 
 			bool bIsDone = InProcessHandle->IsDone();
 
 			if (IngestError.HasError())
 			{
 				ProcessingState = FIngestJob::EProcessingState::Aborted;
-				UE_LOG(LogCaptureManagerIngestJobExecutor, Error, TEXT("[%s] Ingest failed: %s (Error code: %d)"), *ExecutorName, *IngestError.GetError().GetMessage(), IngestError.GetError().GetCode());
+				Message = FText::Format(LOCTEXT("IngestProcessStateMessage_Aborted", "Ingest failed : {0} (Error code : {1})"), FText::FromString(IngestError.GetError().GetMessage()), FText::AsNumber(IngestError.GetError().GetCode()));
+				UE_LOG(LogCaptureManagerIngestJobExecutor, Error, TEXT("[%s] %s"), *ExecutorName, *Message.ToString());
 				bIsDone = true;
 			}
 			else if (bIsDone)
 			{
 				ProcessingState = FIngestJob::EProcessingState::Complete;
 				InJob.SetProgress(1.0);
+				Message = LOCTEXT("IngestProcessStateMessage_Complete", "Success");
 			}
 			else
 			{
 				ProcessingState = FIngestJob::EProcessingState::Running;
+				Message = LOCTEXT("IngestProcessStateMessage_Complete", "Running...");
 			}
 
-			InJob.SetProcessingState(ProcessingState);
+			InJob.SetProcessingState(FIngestJob::FProcessingState{ ProcessingState, MoveTemp(Message) });
 			JobProcessingStateChanged.ExecuteIfBound(InJob.GetGuid(), ProcessingState);
 			
 			if (bIsDone)
@@ -331,3 +341,5 @@ TUniquePtr<FIngestJobExecutor::FRunIngestContext> FIngestJobExecutor::RunIngest(
 }
 
 } // namespace UE::CaptureManager
+
+#undef LOCTEXT_NAMESPACE
diff --git a/Engine/Plugins/VirtualProduction/CaptureManager/CaptureManagerApp/Source/CaptureManagerEditor/Private/IngestManagement/IngestJobQueue.cpp b/Engine/Plugins/VirtualProduction/CaptureManager/CaptureManagerApp/Source/CaptureManagerEditor/Private/IngestManagement/IngestJobQueue.cpp
--- a/Engine/Plugins/VirtualProduction/CaptureManager/CaptureManagerApp/Source/CaptureManagerEditor/Private/IngestManagement/IngestJobQueue.cpp
+++ b/Engine/Plugins/VirtualProduction/CaptureManager/CaptureManagerApp/Source/CaptureManagerEditor/Private/IngestManagement/IngestJobQueue.cpp
@@ -2,8 +2,9 @@
 
 #include "IngestJobQueue.h"
 
-DEFINE_LOG_CATEGORY_STATIC(LogCaptureIngestJobQueue, Log, All);
+#define LOCTEXT_NAMESPACE "IngestJobQueue"
 
+DEFINE_LOG_CATEGORY_STATIC(LogCaptureIngestJobQueue, Log, All);
 
 namespace UE::CaptureManager
 {
@@ -56,7 +57,7 @@ uint32 FIngestJobQueue::CountQueuedDeviceJobs(const FGuid InDeviceId, const uint
 	{
 		if (Job->GetCaptureDeviceId() == InDeviceId)
 		{
-			if (static_cast<uint32>(Job->GetProcessingState()) & InJobsToCountFlags)
+			if (static_cast<uint32>(Job->GetProcessingState().State) & InJobsToCountFlags)
 			{
 				return true;
 			}
@@ -97,9 +98,9 @@ TSharedPtr<FIngestJob> FIngestJobQueue::ClaimFirstPending(const FIngestJob::EPro
 
 	for (const TSharedRef<FIngestJob>& Job : Jobs)
 	{
-		if (Job->GetProcessingState() == FIngestJob::EProcessingState::Pending)
+		if (Job->GetProcessingState().State == FIngestJob::EProcessingState::Pending)
 		{
-			Job->SetProcessingState(ProcessingState);
+			Job->SetProcessingState(FIngestJob::FProcessingState{ ProcessingState, FText(LOCTEXT("IngestProcessStateMessage_Pending","Pending...")) });
 			return Job;
 		}
 	}
@@ -176,3 +177,4 @@ bool FIngestJobQueue::JobCanBeAdded(const FIngestJob& JobToCheck) const
 
 } // namespace UE::CaptureManager
 
+#undef LOCTEXT_NAMESPACE
diff --git a/Engine/Plugins/VirtualProduction/CaptureManager/CaptureManagerApp/Source/CaptureManagerEditor/Private/UI/SIngestJobProcessor.cpp b/Engine/Plugins/VirtualProduction/CaptureManager/CaptureManagerApp/Source/CaptureManagerEditor/Private/UI/SIngestJobProcessor.cpp
--- a/Engine/Plugins/VirtualProduction/CaptureManager/CaptureManagerApp/Source/CaptureManagerEditor/Private/UI/SIngestJobProcessor.cpp
+++ b/Engine/Plugins/VirtualProduction/CaptureManager/CaptureManagerApp/Source/CaptureManagerEditor/Private/UI/SIngestJobProcessor.cpp
@@ -58,7 +58,7 @@ void ResetJobs(TArray<TSharedRef<FIngestJob>> InIngestJobs, TSharedPtr<FIngestJo
 	Algo::TransformIf(InIngestJobs, AbortedJobs,
 					  [](const TSharedRef<FIngestJob>& Input) -> bool
 					  {
-						  return Input->GetProcessingState() == FIngestJob::EProcessingState::Aborted;
+						  return Input->GetProcessingState().State == FIngestJob::EProcessingState::Aborted;
 					  },
 					  [](const TSharedRef<FIngestJob>& Input) -> FGuid
 					  {
@@ -473,24 +473,28 @@ TSharedRef<SWidget> SIngestJobRow::CreateStatusWidget()
 {
 	TSharedPtr<SWidget> StateWidget;
 
-	switch (IngestJob->GetProcessingState())
+	FIngestJob::EProcessingState State = IngestJob->GetProcessingState().State;
+	switch (State)
 	{
 		case FIngestJob::EProcessingState::Running:
 			StateWidget = SNew(SImage)
 				.Image(FAppStyle::Get().GetBrush("Icons.ArrowRight"))
-				.ColorAndOpacity(FStyleColors::Success);
+				.ColorAndOpacity(FStyleColors::Success)
+				.ToolTipText(IngestJob->GetProcessingState().Message);
 			break;
 
 		case FIngestJob::EProcessingState::Complete:
 			StateWidget = SNew(SImage)
 				.Image(FAppStyle::Get().GetBrush("Icons.Success"))
-				.ColorAndOpacity(FStyleColors::Success);
+				.ColorAndOpacity(FStyleColors::Success)
+				.ToolTipText(IngestJob->GetProcessingState().Message);
 			break;
 
 		case FIngestJob::EProcessingState::Aborted:
 			StateWidget = SNew(SImage)
 				.Image(FAppStyle::Get().GetBrush("Icons.Alert"))
-				.ColorAndOpacity(FStyleColors::Error);
+				.ColorAndOpacity(FStyleColors::Error)
+				.ToolTipText(IngestJob->GetProcessingState().Message);
 			break;
 
 		default:
@@ -510,7 +514,7 @@ TSharedRef<SWidget> SIngestJobRow::CreateStatusWidget()
 
 TSharedRef<SWidget> SIngestJobRow::CreateActionWidget()
 {
-	if (IngestJob->GetProcessingState() == FIngestJob::EProcessingState::Aborted)
+	if (IngestJob->GetProcessingState().State == FIngestJob::EProcessingState::Aborted)
 	{
 		return SNew(SScaleBox)
 			.Stretch(EStretch::ScaleToFit)
