[CaptureManager] Fixed localization text clash

diff --git a/Engine/Plugins/VirtualProduction/CaptureManager/CaptureManagerApp/Source/CaptureManagerEditor/Private/IngestManagement/IngestJobExecutor.cpp b/Engine/Plugins/VirtualProduction/CaptureManager/CaptureManagerApp/Source/CaptureManagerEditor/Private/IngestManagement/IngestJobExecutor.cpp
--- a/Engine/Plugins/VirtualProduction/CaptureManager/CaptureManagerApp/Source/CaptureManagerEditor/Private/IngestManagement/IngestJobExecutor.cpp
+++ b/Engine/Plugins/VirtualProduction/CaptureManager/CaptureManagerApp/Source/CaptureManagerEditor/Private/IngestManagement/IngestJobExecutor.cpp
@@ -302,7 +302,7 @@ TUniquePtr<FIngestJobExecutor::FRunIngestContext> FIngestJobExecutor::RunIngest(
 			else
 			{
 				ProcessingState = FIngestJob::EProcessingState::Running;
-				Message = LOCTEXT("IngestProcessStateMessage_Complete", "Running...");
+				Message = LOCTEXT("IngestProcessStateMessage_Running", "Running...");
 			}
 
 			InJob.SetProcessingState(FIngestJob::FProcessingState{ ProcessingState, MoveTemp(Message) });
