[Capture Manager] Fixed truncated paths in timecode extraction log.

#rnx

diff --git a/Engine/Plugins/VirtualProduction/CaptureManager/CaptureManagerApp/Source/IngestLiveLinkDevice/Private/Utils/CaptureExtractTimecode.cpp b/Engine/Plugins/VirtualProduction/CaptureManager/CaptureManagerApp/Source/IngestLiveLinkDevice/Private/Utils/CaptureExtractTimecode.cpp
--- a/Engine/Plugins/VirtualProduction/CaptureManager/CaptureManagerApp/Source/IngestLiveLinkDevice/Private/Utils/CaptureExtractTimecode.cpp
+++ b/Engine/Plugins/VirtualProduction/CaptureManager/CaptureManagerApp/Source/IngestLiveLinkDevice/Private/Utils/CaptureExtractTimecode.cpp
@@ -21,7 +21,10 @@ DEFINE_LOG_CATEGORY_STATIC(LogCaptureExtractTimecode, Log, All);
 
 FCaptureExtractVideoTimecode::FCaptureExtractVideoTimecode(const FString& InFilePath)
 	: FilePath(InFilePath)
-{}
+{
+	// Replace any backslashes to avoid escape code issues while logging the file path
+	FilePath = FilePath.Replace(TEXT("\\"), TEXT("/"));
+}
 
 UE::CaptureManager::FTimecodeAndRateResult FCaptureExtractVideoTimecode::Extract()
 {
