Image Sequence reader implementation

diff --git a/Engine/Plugins/VirtualProduction/CaptureManager/CaptureManagerApp/Source/CaptureDataConverter/Private/Nodes/CaptureConvertDepthData.cpp b/Engine/Plugins/VirtualProduction/CaptureManager/CaptureManagerApp/Source/CaptureDataConverter/Private/Nodes/CaptureConvertDepthData.cpp
--- a/Engine/Plugins/VirtualProduction/CaptureManager/CaptureManagerApp/Source/CaptureDataConverter/Private/Nodes/CaptureConvertDepthData.cpp
+++ b/Engine/Plugins/VirtualProduction/CaptureManager/CaptureManagerApp/Source/CaptureDataConverter/Private/Nodes/CaptureConvertDepthData.cpp
@@ -167,8 +167,10 @@ FCaptureConvertDepthData::FResult FCaptureConvertDepthData::CopyData()
 
 	FString DestinationDirectory = OutputDirectory / Depth.Name;
 	FCaptureConvertDepthData::FResult Result = MakeValue();
+	int32 NumberOfCopiedFiles = 0;
+
 	IImageWrapperModule& ImageWrapperModule = FModuleManager::Get().LoadModuleChecked<IImageWrapperModule>(TEXT("ImageWrapper"));
-	FileManager.IterateDirectoryRecursively(*DepthFolderPath, [this, &TaskProgress, &Result, DestinationDirectory, &FileManager, &ImageWrapperModule](const TCHAR* InFileName, bool bInIsDirectory)
+	FileManager.IterateDirectoryRecursively(*DepthFolderPath, [this, &TaskProgress, &Result, DestinationDirectory, &FileManager, &ImageWrapperModule, &NumberOfCopiedFiles](const TCHAR* InFileName, bool bInIsDirectory)
 	{
 		if (bInIsDirectory)
 		{
@@ -177,9 +179,8 @@ FCaptureConvertDepthData::FResult FCaptureConvertDepthData::CopyData()
 
 		if (ImageWrapperModule.GetImageFormatFromExtension(InFileName) != EImageFormat::EXR)
 		{
-			FText Message = LOCTEXT("CaptureConvertDepthData_UnsupportedFileFormat", "Unsupported file format");
-			Result = MakeError(MoveTemp(Message));
-			return false;
+			UE_LOG(LogCaptureConvertDepthData, Warning, TEXT("Image file format is unsupported %s"), InFileName);
+			return true;
 		}
 		
 		FTaskProgress::FTask Task = TaskProgress->StartTask();
@@ -204,10 +205,18 @@ FCaptureConvertDepthData::FResult FCaptureConvertDepthData::CopyData()
 			return false;
 		}
 
+		++NumberOfCopiedFiles;
+
 		return true;
 	});
 
-	return MakeValue();
+	if (NumberOfCopiedFiles == 0)
+	{
+		FText Message = FText::Format(LOCTEXT("CaptureConvertDepthData_NoFiles", "Failed to copy from {0} to {1}"), FText::FromString(DepthFolderPath), FText::FromString(DestinationDirectory));
+		return MakeError(MoveTemp(Message));
+	}
+
+	return Result;
 }
 
 FCaptureConvertDepthData::FWriteResult FCaptureConvertDepthData::OnWrite(FWritingContext InContext)
diff --git a/Engine/Plugins/VirtualProduction/CaptureManager/CaptureManagerApp/Source/CaptureDataConverter/Private/Nodes/CaptureConvertVideoData.cpp b/Engine/Plugins/VirtualProduction/CaptureManager/CaptureManagerApp/Source/CaptureDataConverter/Private/Nodes/CaptureConvertVideoData.cpp
--- a/Engine/Plugins/VirtualProduction/CaptureManager/CaptureManagerApp/Source/CaptureDataConverter/Private/Nodes/CaptureConvertVideoData.cpp
+++ b/Engine/Plugins/VirtualProduction/CaptureManager/CaptureManagerApp/Source/CaptureDataConverter/Private/Nodes/CaptureConvertVideoData.cpp
@@ -47,13 +47,13 @@ FCaptureConvertVideoData::FResult FCaptureConvertVideoData::Run()
 	{
 		FTakeMetadataPathUtils::ValidatePathType(Video.Path, Video.PathType.GetValue());
 	}
-
-	if (Video.PathType.IsSet() && Video.PathType.GetValue() == FTakeMetadata::FVideo::EPathType::File)
+	
+	if (ShouldCopy())
 	{
-		return ConvertData();
+		return CopyData();
 	}
-	
-	return CopyData();
+
+	return ConvertData();
 }
 
 FCaptureConvertVideoData::FResult FCaptureConvertVideoData::ConvertData()
@@ -68,21 +68,58 @@ FCaptureConvertVideoData::FResult FCaptureConvertVideoData::ConvertData()
 	FCaptureManagerMediaRWModule& MediaRWModule =
 		FModuleManager::LoadModuleChecked<FCaptureManagerMediaRWModule>("CaptureManagerMediaRW");
 
-	TValueOrError<TUniquePtr<IVideoReader>, FText> VideoReaderResult = MediaRWModule.Get().CreateVideoReader(VideoFilePath);
+	TUniquePtr<IVideoReader> VideoReader = nullptr;
+	if (Video.PathType.IsSet() && Video.PathType.GetValue() == FTakeMetadata::FVideo::EPathType::Folder)
+	{
+		static constexpr FStringView Format = TEXT("image_sequence");
+		VideoReader = MediaRWModule.Get().CreateVideoReaderByFormat(FString(Format));
+
+		if (!VideoReader)
+		{
+			FText Message = FText::Format(
+				LOCTEXT("CaptureConvertVideo_ReaderIsntRegistered", "Video reader for image sequences isn't registered {0}"), FText::FromStringView(Format));
+			return MakeError(MoveTemp(Message));
+		}
+
+		TOptional<FText> Error = VideoReader->Open(VideoFilePath);
+
+		if (Error.IsSet())
+		{
+			return MakeError(MoveTemp(Error.GetValue()));
+		}
+	}
+	else
+	{
+		TValueOrError<TUniquePtr<IVideoReader>, FText> VideoReaderResult = MediaRWModule.Get().CreateVideoReader(VideoFilePath);
+
+		if (VideoReaderResult.HasError())
+		{
+			FText Message = FText::Format(LOCTEXT("CaptureConvertVideo_UnsupportedFile", "Video file format is unsupported {0}. Consider enabling Third Party Encoder in Capture Manager settings."), FText::FromString(VideoFilePath));
+			return MakeError(MoveTemp(Message));
+		}
+
+		VideoReader = VideoReaderResult.StealValue();
+	}
+	
+	check(VideoReader);
+
+	ON_SCOPE_EXIT
+	{
+		VideoReader->Close();
+	};
+
 	TValueOrError<TUniquePtr<IImageWriter>, FText> ImageWriterResult = MediaRWModule.Get().CreateImageWriter(DestinationDirectory, VideoParams.ImageFileName, VideoParams.Format);
 
-	if (VideoReaderResult.HasError() || ImageWriterResult.HasError())
+	if (ImageWriterResult.HasError())
 	{
 		FText Message = FText::Format(LOCTEXT("CaptureConvertVideo_UnsupportedFile", "Video file format is unsupported {0}. Consider enabling Third Party Encoder in Capture Manager settings."), FText::FromString(VideoFilePath));
 		return MakeError(MoveTemp(Message));
 	}
 
-	TUniquePtr<IVideoReader> VideoReader = VideoReaderResult.StealValue();
 	TUniquePtr<IImageWriter> ImageWriter = ImageWriterResult.StealValue();
 
 	ON_SCOPE_EXIT
 	{
-		VideoReader->Close();
 		ImageWriter->Close();
 	};
 
@@ -124,6 +161,75 @@ FCaptureConvertVideoData::FResult FCaptureConvertVideoData::ConvertData()
 	return Result;
 }
 
+bool FCaptureConvertVideoData::ShouldCopy() const
+{
+	FCaptureManagerMediaRWModule& MediaRWModule =
+		FModuleManager::LoadModuleChecked<FCaptureManagerMediaRWModule>("CaptureManagerMediaRW");
+
+	if (!Video.PathType.IsSet() || Video.PathType.GetValue() != FTakeMetadata::FVideo::EPathType::Folder)
+	{
+		return false;
+	}
+
+	TUniquePtr<IVideoReader> VideoReader = MediaRWModule.Get().CreateVideoReaderByFormat(TEXT("image_sequence"));
+
+	if (!VideoReader)
+	{
+		return false;
+	}
+
+	const FString VideoFilePath = FPaths::ConvertRelativePathToFull(Params.TakeOriginDirectory, Video.Path);
+	TOptional<FText> Error = VideoReader->Open(VideoFilePath);
+
+	if (Error.IsSet())
+	{
+		return false;
+	}
+
+	TValueOrError<TUniquePtr<UE::CaptureManager::FMediaTextureSample>, FText> ReadResult = VideoReader->Next();
+	if (!ReadResult.IsValid())
+	{
+		return false;
+	}
+
+	TUniquePtr<UE::CaptureManager::FMediaTextureSample> Sample = ReadResult.StealValue();
+
+	if (!Sample)
+	{
+		return false;
+	}
+
+	IFileManager& FileManager = IFileManager::Get();
+	IImageWrapperModule& ImageWrapperModule =
+		FModuleManager::Get().LoadModuleChecked<IImageWrapperModule>(TEXT("ImageWrapper"));
+	FString Extension;
+	FileManager.IterateDirectory(*VideoFilePath, [&Extension, &ImageWrapperModule](const TCHAR* InFilenameOrDirectory, bool bInIsDirectory)
+	{
+		if (bInIsDirectory)
+		{
+			return true;
+		}
+
+		Extension = FPaths::GetExtension(InFilenameOrDirectory, false);
+		if (ImageWrapperModule.GetImageFormatFromExtension(*Extension) != EImageFormat::Invalid)
+		{
+			// Found the correct extension and return false to exit the iteration
+			return false;
+		}
+
+		return true;
+	});
+
+	if (Sample->CurrentFormat == VideoParams.OutputPixelFormat &&
+		VideoParams.Rotation == EMediaOrientation::Original &&
+		VideoParams.Format == Extension)
+	{
+		return true;
+	}
+
+	return false;
+}
+
 FCaptureConvertVideoData::FResult FCaptureConvertVideoData::CopyData()
 {
 	using namespace UE::CaptureManager;
@@ -157,9 +263,10 @@ FCaptureConvertVideoData::FResult FCaptureConvertVideoData::CopyData()
 	FString DestinationDirectory = OutputDirectory / Video.Name;
 
 	FCaptureConvertVideoData::FResult Result = MakeValue();
+	int32 NumberOfCopiedFiles = 0;
 
 	IImageWrapperModule& ImageWrapperModule = FModuleManager::Get().LoadModuleChecked<IImageWrapperModule>(TEXT("ImageWrapper"));
-	FileManager.IterateDirectoryRecursively(*VideoFolderPath, [this, &TaskProgress, &Result, DestinationDirectory, &FileManager, &ImageWrapperModule, VideoFolderPath](const TCHAR* InFileName, bool bInIsDirectory)
+	FileManager.IterateDirectoryRecursively(*VideoFolderPath, [this, &TaskProgress, &Result, DestinationDirectory, &FileManager, &ImageWrapperModule, &NumberOfCopiedFiles, VideoFolderPath](const TCHAR* InFileName, bool bInIsDirectory)
 	{
 		if (bInIsDirectory)
 		{
@@ -168,10 +275,8 @@ FCaptureConvertVideoData::FResult FCaptureConvertVideoData::CopyData()
 
 		if (ImageWrapperModule.GetImageFormatFromExtension(InFileName) == EImageFormat::Invalid)
 		{
-			FText Message = FText::Format(LOCTEXT("CaptureConvertVideoData_UnsupportedFileFormat", "Image file format is unsupported {0}"), 
-										  FText::FromString(InFileName));
-			Result = MakeError(MoveTemp(Message));
-			return false;
+			UE_LOG(LogCaptureConvertVideoData, Warning, TEXT("Image file format is unsupported %s"), InFileName);
+			return true;
 		}
 		
 		FTaskProgress::FTask Task = TaskProgress->StartTask();
@@ -196,10 +301,18 @@ FCaptureConvertVideoData::FResult FCaptureConvertVideoData::CopyData()
 			return false;
 		}
 
+		++NumberOfCopiedFiles;
+
 		return true;
 	});
 
-	return MakeValue();
+	if (NumberOfCopiedFiles == 0)
+	{
+		FText Message = FText::Format(LOCTEXT("CaptureConvertVideoData_NoFiles", "Failed to copy from {0} to {1}"), FText::FromString(VideoFolderPath), FText::FromString(DestinationDirectory));
+		return MakeError(MoveTemp(Message));
+	}
+
+	return Result;
 }
 
 FCaptureConvertVideoData::FWriteResult FCaptureConvertVideoData::OnWrite(FWritingContext InContext)
diff --git a/Engine/Plugins/VirtualProduction/CaptureManager/CaptureManagerApp/Source/CaptureDataConverter/Private/Nodes/CaptureConvertVideoData.h b/Engine/Plugins/VirtualProduction/CaptureManager/CaptureManagerApp/Source/CaptureDataConverter/Private/Nodes/CaptureConvertVideoData.h
--- a/Engine/Plugins/VirtualProduction/CaptureManager/CaptureManagerApp/Source/CaptureDataConverter/Private/Nodes/CaptureConvertVideoData.h
+++ b/Engine/Plugins/VirtualProduction/CaptureManager/CaptureManagerApp/Source/CaptureDataConverter/Private/Nodes/CaptureConvertVideoData.h
@@ -28,6 +28,7 @@ private:
 	virtual FResult Run() override;
 
 	FResult ConvertData();
+	bool ShouldCopy() const;
 	FResult CopyData();
 
 	EMediaOrientation ConvertOrientation(FTakeMetadata::FVideo::EOrientation InOrientation) const;
diff --git a/Engine/Plugins/VirtualProduction/CaptureManager/CaptureManagerApp/Source/CaptureDataConverter/Private/Nodes/ThirdPartyEncoder/CaptureConvertVideoDataThirdParty.cpp b/Engine/Plugins/VirtualProduction/CaptureManager/CaptureManagerApp/Source/CaptureDataConverter/Private/Nodes/ThirdPartyEncoder/CaptureConvertVideoDataThirdParty.cpp
--- a/Engine/Plugins/VirtualProduction/CaptureManager/CaptureManagerApp/Source/CaptureDataConverter/Private/Nodes/ThirdPartyEncoder/CaptureConvertVideoDataThirdParty.cpp
+++ b/Engine/Plugins/VirtualProduction/CaptureManager/CaptureManagerApp/Source/CaptureDataConverter/Private/Nodes/ThirdPartyEncoder/CaptureConvertVideoDataThirdParty.cpp
@@ -318,9 +318,10 @@ FCaptureConvertVideoDataThirdParty::FResult FCaptureConvertVideoDataThirdParty::
 	FString DestinationDirectory = OutputDirectory / Video.Name;
 
 	FCaptureConvertVideoDataThirdParty::FResult Result = MakeValue();
+	int32 NumberOfCopiedFiles = 0;
 
 	IImageWrapperModule& ImageWrapperModule = FModuleManager::Get().LoadModuleChecked<IImageWrapperModule>(TEXT("ImageWrapper"));
-	FileManager.IterateDirectoryRecursively(*VideoFolderPath, [this, &TaskProgress, &Result, DestinationDirectory, &FileManager, &ImageWrapperModule](const TCHAR* InFileName, bool bInIsDirectory)
+	FileManager.IterateDirectoryRecursively(*VideoFolderPath, [this, &TaskProgress, &Result, DestinationDirectory, &FileManager, &ImageWrapperModule, &NumberOfCopiedFiles](const TCHAR* InFileName, bool bInIsDirectory)
 	{
 		if (bInIsDirectory)
 		{
@@ -329,10 +330,8 @@ FCaptureConvertVideoDataThirdParty::FResult FCaptureConvertVideoDataThirdParty::
 
 		if (ImageWrapperModule.GetImageFormatFromExtension(InFileName) == EImageFormat::Invalid)
 		{
-			FText Message = FText::Format(LOCTEXT("CaptureConvertVideoDataTP_UnsupportedFileFormat", "Image file format is unsupported {0}"), 
-										  FText::FromString(InFileName));
-			Result = MakeError(MoveTemp(Message));
-			return false;
+			UE_LOG(LogCaptureConvertVideoDataThirdParty, Warning, TEXT("Image file format is unsupported %s"), InFileName);
+			return true;
 		}
 		
 		FTaskProgress::FTask Task = TaskProgress->StartTask();
@@ -361,10 +360,18 @@ FCaptureConvertVideoDataThirdParty::FResult FCaptureConvertVideoDataThirdParty::
 			return false;
 		}
 
+		++NumberOfCopiedFiles;
+
 		return true;
 	});
 
-	return MakeValue();
+	if (NumberOfCopiedFiles == 0)
+	{
+		FText Message = FText::Format(LOCTEXT("CaptureConvertVideoNodeTP_NoFiles", "Failed to copy from {0} to {1}"), FText::FromString(VideoFolderPath), FText::FromString(DestinationDirectory));
+		return MakeError(MoveTemp(Message));
+	}
+
+	return Result;
 }
 
 #undef LOCTEXT_NAMESPACE
diff --git a/Engine/Plugins/VirtualProduction/CaptureManager/CaptureManagerApp/Source/CaptureManagerMediaRW/Private/CaptureManagerMediaRWModule.cpp b/Engine/Plugins/VirtualProduction/CaptureManager/CaptureManagerApp/Source/CaptureManagerMediaRW/Private/CaptureManagerMediaRWModule.cpp
--- a/Engine/Plugins/VirtualProduction/CaptureManager/CaptureManagerApp/Source/CaptureManagerMediaRW/Private/CaptureManagerMediaRWModule.cpp
+++ b/Engine/Plugins/VirtualProduction/CaptureManager/CaptureManagerApp/Source/CaptureManagerMediaRW/Private/CaptureManagerMediaRWModule.cpp
@@ -10,6 +10,7 @@
 #include "Readers/MHADepthVideoReader.h"
 #include "Readers/MHAICalibrationReader.h"
 #include "Readers/OpenCVCalibrationReader.h"
+#include "Readers/ImageSequenceReader.h"
 
 #include "Writers/AudioWaveMediaWriter.h"
 #include "Writers/DepthImageWriter.h"
@@ -22,18 +23,19 @@ void FCaptureManagerMediaRWModule::StartupModule()
 	// Readers
 #if PLATFORM_WINDOWS && !UE_SERVER
 	FWindowsRWHelpers::Init();
-	FWindowsRWHelpers::RegisterReaders(MediaRWManager.Get());
-	FWindowsRWHelpers::RegisterWriters(MediaRWManager.Get());
+	FWindowsRWHelpers::RegisterReaders(*MediaRWManager);
+	FWindowsRWHelpers::RegisterWriters(*MediaRWManager);
 #endif
 
-	FMHADepthVideoReaderHelpers::RegisterReaders(MediaRWManager.Get());
-	FMHAICalibrationReaderHelpers::RegisterReaders(MediaRWManager.Get());
-	FOpenCvCalibrationReaderHelpers::RegisterReaders(MediaRWManager.Get());
+	FMHADepthVideoReaderHelpers::RegisterReaders(*MediaRWManager);
+	FMHAICalibrationReaderHelpers::RegisterReaders(*MediaRWManager);
+	FOpenCvCalibrationReaderHelpers::RegisterReaders(*MediaRWManager);
+	FImageSequenceReaderHelper::RegisterReaders(*MediaRWManager);
 
 	// Writers
-	FAudioWaveWriterHelpers::RegisterWriters(MediaRWManager.Get());
-	FDepthExrImageWriterHelpers::RegisterWriters(MediaRWManager.Get());
-	FUnrealCalibrationWriterHelpers::RegisterWriters(MediaRWManager.Get());
+	FAudioWaveWriterHelpers::RegisterWriters(*MediaRWManager);
+	FDepthExrImageWriterHelpers::RegisterWriters(*MediaRWManager);
+	FUnrealCalibrationWriterHelpers::RegisterWriters(*MediaRWManager);
 }
 
 void FCaptureManagerMediaRWModule::ShutdownModule()
diff --git a/Engine/Plugins/VirtualProduction/CaptureManager/CaptureManagerApp/Source/CaptureManagerMediaRW/Private/Readers/ImageSequenceReader.cpp b/Engine/Plugins/VirtualProduction/CaptureManager/CaptureManagerApp/Source/CaptureManagerMediaRW/Private/Readers/ImageSequenceReader.cpp
new file mode 100644
--- /dev/null
+++ b/Engine/Plugins/VirtualProduction/CaptureManager/CaptureManagerApp/Source/CaptureManagerMediaRW/Private/Readers/ImageSequenceReader.cpp
@@ -0,0 +1,213 @@
+// Copyright Epic Games, Inc. All Rights Reserved.
+
+#include "ImageSequenceReader.h"
+
+#include "MediaRWManager.h"
+
+#include "Misc/FileHelper.h"
+#include "IImageWrapperModule.h"
+#include "Modules/ModuleManager.h"
+
+#define LOCTEXT_NAMESPACE "ImageSequenceReader"
+
+namespace UE::CaptureManager::Private
+{
+
+static TArray<FString> GetImageSequenceFilesFromPath(const FString& InFullSequencePath)
+{
+	TArray<FString> ImageFiles;
+
+	IImageWrapperModule& ImageWrapperModule = FModuleManager::LoadModuleChecked<IImageWrapperModule>(TEXT("ImageWrapper"));
+	IFileManager& FileManager = IFileManager::Get();
+
+	bool bIterateResult = FileManager.IterateDirectory(*InFullSequencePath, [&ImageFiles, &ImageWrapperModule](const TCHAR* InFilenameOrDirectory, bool bInIsDirectory)
+	{
+		if (!bInIsDirectory)
+		{
+			EImageFormat Format = ImageWrapperModule.GetImageFormatFromExtension(*FPaths::GetExtension(InFilenameOrDirectory));
+			if (Format != EImageFormat::Invalid)
+			{
+				ImageFiles.Add(InFilenameOrDirectory);
+			}
+		}
+
+		return true;
+	});
+
+	return ImageFiles;
+}
+
+static bool GetImageDimensionsFromPath(const FString& InImagePath, FIntPoint& OutDimensions)
+{
+	TArray<uint8> RawFileData;
+	if (FFileHelper::LoadFileToArray(RawFileData, *InImagePath))
+	{
+		IImageWrapperModule& ImageWrapperModule = FModuleManager::LoadModuleChecked<IImageWrapperModule>(TEXT("ImageWrapper"));
+		EImageFormat ImageFormat = ImageWrapperModule.DetectImageFormat(RawFileData.GetData(), RawFileData.Num());
+		TSharedPtr<IImageWrapper> ImageWrapper = ImageWrapperModule.CreateImageWrapper(ImageFormat);
+
+		if (ImageWrapper.IsValid() && ImageWrapper->SetCompressed(RawFileData.GetData(), RawFileData.Num()))
+		{
+			OutDimensions.X = ImageWrapper->GetWidth();
+			OutDimensions.Y = ImageWrapper->GetHeight();
+			return true;
+		}
+	}
+
+	return false;
+}
+
+static UE::CaptureManager::EMediaTexturePixelFormat ConvertFormat(ERGBFormat InFormat, int32 InBitDepth)
+{
+	switch (InFormat)
+	{
+		case ERGBFormat::Gray:
+			if (InBitDepth == 8)
+			{
+				return UE::CaptureManager::EMediaTexturePixelFormat::U8_Mono;
+			}
+			else if (InBitDepth == 16)
+			{
+				return UE::CaptureManager::EMediaTexturePixelFormat::U16_Mono;
+			}
+			else
+			{
+				return UE::CaptureManager::EMediaTexturePixelFormat::Undefined;
+			}
+		case ERGBFormat::GrayF:
+			return UE::CaptureManager::EMediaTexturePixelFormat::F_Mono;
+		case ERGBFormat::RGBA:
+			return UE::CaptureManager::EMediaTexturePixelFormat::U8_RGBA;
+		case ERGBFormat::BGRA:
+			return UE::CaptureManager::EMediaTexturePixelFormat::U8_BGRA;
+		default:
+			return UE::CaptureManager::EMediaTexturePixelFormat::Undefined;
+	}
+}
+
+}
+
+void FImageSequenceReaderHelper::RegisterReaders(FMediaRWManager& InManager)
+{
+	TArray<FString> SupportedFormats = { TEXT("image_sequence") };
+	InManager.RegisterVideoReader(SupportedFormats, MakeUnique<FImageSequenceReaderFactory>());
+}
+
+TUniquePtr<IVideoReader> FImageSequenceReaderFactory::CreateVideoReader()
+{
+	return MakeUnique<FImageSequenceReader>();
+}
+
+FImageSequenceReader::FImageSequenceReader() = default;
+FImageSequenceReader::~FImageSequenceReader() = default;
+
+TOptional<FText> FImageSequenceReader::Open(const FString& InDirectoryPath)
+{
+	if (!FPaths::DirectoryExists(InDirectoryPath))
+	{
+		FText Error = FText::Format(
+			LOCTEXT("ImageSequenceFolderMissing", "Failed to read image sequence directory: {0}"),
+			FText::FromString(InDirectoryPath));
+		return Error;
+	}
+
+	ImagePaths = UE::CaptureManager::Private::GetImageSequenceFilesFromPath(InDirectoryPath);
+
+	if (ImagePaths.IsEmpty())
+	{
+		FText Error = FText::Format(
+			LOCTEXT("ImageSequenceEmptyDirectory", "No images found in the directory: {0}"),
+			FText::FromString(InDirectoryPath));
+		return Error;
+	}
+
+	ImagePaths.Sort();
+
+	check(ImagePaths.IsValidIndex(0));
+	UE::CaptureManager::Private::GetImageDimensionsFromPath(ImagePaths[0], Dimensions);
+
+	return {};
+}
+
+TOptional<FText> FImageSequenceReader::Close()
+{
+	ImagePaths.Empty();
+
+	return {};
+}
+
+TValueOrError<TUniquePtr<UE::CaptureManager::FMediaTextureSample>, FText> FImageSequenceReader::Next()
+{
+	const int32 CurrentFrame = CurrentFrameNumber;
+
+	if (!ImagePaths.IsValidIndex(CurrentFrame))
+	{
+		return MakeValue(nullptr);
+	}
+
+	const FString& ImagePath = ImagePaths[CurrentFrame];
+
+	TArray<uint8> RawFileData;
+	if (!FFileHelper::LoadFileToArray(RawFileData, *ImagePath))
+	{
+		FText Error = FText::Format(LOCTEXT("ImageSequenceNextFrameReadFailed", "Failed to read frame {0} from {1}"),
+									FText::AsNumber(CurrentFrame), FText::FromString(ImagePath));
+		return MakeError(MoveTemp(Error));
+	}
+
+	TUniquePtr<UE::CaptureManager::FMediaTextureSample> Sample = 
+		MakeUnique<UE::CaptureManager::FMediaTextureSample>();
+
+	Sample->Dimensions = Dimensions;
+	Sample->Stride = Dimensions.X;
+	Sample->Time = FTimespan::FromSeconds(CurrentFrame);
+
+	IImageWrapperModule& ImageWrapperModule = FModuleManager::LoadModuleChecked<IImageWrapperModule>(TEXT("ImageWrapper"));
+	EImageFormat ImageFormat = ImageWrapperModule.DetectImageFormat(RawFileData.GetData(), RawFileData.Num());
+	TSharedPtr<IImageWrapper> ImageWrapper = ImageWrapperModule.CreateImageWrapper(ImageFormat);
+
+	if (!ImageWrapper.IsValid() || !ImageWrapper->SetCompressed(RawFileData.GetData(), RawFileData.Num()))
+	{
+		FText Error = FText::Format(LOCTEXT("ImageSequenceNextDataToImageFailed", "Unsupported image format {0} for {1}"),
+									FText::FromString(FPaths::GetExtension(ImagePath, false)), FText::FromString(ImagePath));
+
+		return MakeError(MoveTemp(Error));
+	}
+
+	// Emptying the first buffer before obtaining a new one.
+	RawFileData.Empty();
+
+	Sample->CurrentFormat = UE::CaptureManager::Private::ConvertFormat(ImageWrapper->GetFormat(), ImageWrapper->GetBitDepth());
+
+	TArray64<uint8> Buffer;
+	if (!ImageWrapper->GetRaw(Buffer))
+	{
+		FText Error = FText::Format(LOCTEXT("ImageSequenceNextDecodeFailed", "Failed to decode the image {0}"),
+									FText::FromString(ImagePath));
+
+		return MakeError(MoveTemp(Error));
+	}
+
+	Sample->Buffer = MoveTemp(Buffer);
+
+	++CurrentFrameNumber;
+
+	return MakeValue(MoveTemp(Sample));
+}
+
+FTimespan FImageSequenceReader::GetDuration() const
+{
+	return FTimespan::FromSeconds(ImagePaths.Num());
+}
+
+FIntPoint FImageSequenceReader::GetDimensions() const
+{
+	return Dimensions;
+}
+
+FFrameRate FImageSequenceReader::GetFrameRate() const
+{
+	return FFrameRate(1, 1);
+}
+
+#undef LOCTEXT_NAMESPACE
\ No newline at end of file
diff --git a/Engine/Plugins/VirtualProduction/CaptureManager/CaptureManagerApp/Source/CaptureManagerMediaRW/Private/Readers/ImageSequenceReader.h b/Engine/Plugins/VirtualProduction/CaptureManager/CaptureManagerApp/Source/CaptureManagerMediaRW/Private/Readers/ImageSequenceReader.h
new file mode 100644
--- /dev/null
+++ b/Engine/Plugins/VirtualProduction/CaptureManager/CaptureManagerApp/Source/CaptureManagerMediaRW/Private/Readers/ImageSequenceReader.h
@@ -0,0 +1,46 @@
+// Copyright Epic Games, Inc. All Rights Reserved.
+
+#pragma once
+
+#include "IMediaReader.h"
+#include "IMediaRWFactory.h"
+
+#include "GenericPlatform/GenericPlatformFile.h"
+
+class FImageSequenceReaderHelper
+{
+public:
+
+	static void RegisterReaders(class FMediaRWManager& InManager);
+};
+
+class FImageSequenceReaderFactory final : public IVideoReaderFactory
+{
+public:
+
+	virtual TUniquePtr<IVideoReader> CreateVideoReader() override;
+};
+
+class FImageSequenceReader final : public IVideoReader
+{
+public:
+
+	FImageSequenceReader();
+	virtual ~FImageSequenceReader() override;
+
+	virtual TOptional<FText> Open(const FString& InDirectoryPath) override;
+	virtual TOptional<FText> Close() override;
+
+	virtual TValueOrError<TUniquePtr<UE::CaptureManager::FMediaTextureSample>, FText> Next() override;
+
+	// Duration and Frame Rate will be set to number of frames and 1 frame per second
+	virtual FTimespan GetDuration() const override;
+	virtual FIntPoint GetDimensions() const override;
+	virtual FFrameRate GetFrameRate() const override;
+
+private:
+
+	TArray<FString> ImagePaths;
+	FIntPoint Dimensions;
+	std::atomic_int CurrentFrameNumber = 0;
+};
\ No newline at end of file
diff --git a/Engine/Plugins/VirtualProduction/CaptureManager/CaptureManagerApp/Source/CaptureManagerMediaRW/Private/Readers/MHADepthVideoReader.cpp b/Engine/Plugins/VirtualProduction/CaptureManager/CaptureManagerApp/Source/CaptureManagerMediaRW/Private/Readers/MHADepthVideoReader.cpp
--- a/Engine/Plugins/VirtualProduction/CaptureManager/CaptureManagerApp/Source/CaptureManagerMediaRW/Private/Readers/MHADepthVideoReader.cpp
+++ b/Engine/Plugins/VirtualProduction/CaptureManager/CaptureManagerApp/Source/CaptureManagerMediaRW/Private/Readers/MHADepthVideoReader.cpp
@@ -38,10 +38,10 @@ struct FFrameHeader
 };
 #pragma pack(pop)
 
-void FMHADepthVideoReaderHelpers::RegisterReaders(FMediaRWManager* InManager)
+void FMHADepthVideoReaderHelpers::RegisterReaders(FMediaRWManager& InManager)
 {
 	TArray<FString> SupportedFormats = { TEXT("mha_depth") };
-	InManager->RegisterVideoReader(SupportedFormats, MakeUnique<FMHADepthVideoReaderFactory>());
+	InManager.RegisterVideoReader(SupportedFormats, MakeUnique<FMHADepthVideoReaderFactory>());
 }
 
 TUniquePtr<IVideoReader> FMHADepthVideoReaderFactory::CreateVideoReader()
diff --git a/Engine/Plugins/VirtualProduction/CaptureManager/CaptureManagerApp/Source/CaptureManagerMediaRW/Private/Readers/MHADepthVideoReader.h b/Engine/Plugins/VirtualProduction/CaptureManager/CaptureManagerApp/Source/CaptureManagerMediaRW/Private/Readers/MHADepthVideoReader.h
--- a/Engine/Plugins/VirtualProduction/CaptureManager/CaptureManagerApp/Source/CaptureManagerMediaRW/Private/Readers/MHADepthVideoReader.h
+++ b/Engine/Plugins/VirtualProduction/CaptureManager/CaptureManagerApp/Source/CaptureManagerMediaRW/Private/Readers/MHADepthVideoReader.h
@@ -11,7 +11,7 @@ class FMHADepthVideoReaderHelpers
 {
 public:
 
-	static void RegisterReaders(class FMediaRWManager* InManager);
+	static void RegisterReaders(class FMediaRWManager& InManager);
 };
 
 class FMHADepthVideoReaderFactory final : public IVideoReaderFactory
diff --git a/Engine/Plugins/VirtualProduction/CaptureManager/CaptureManagerApp/Source/CaptureManagerMediaRW/Private/Readers/MHAICalibrationReader.cpp b/Engine/Plugins/VirtualProduction/CaptureManager/CaptureManagerApp/Source/CaptureManagerMediaRW/Private/Readers/MHAICalibrationReader.cpp
--- a/Engine/Plugins/VirtualProduction/CaptureManager/CaptureManagerApp/Source/CaptureManagerMediaRW/Private/Readers/MHAICalibrationReader.cpp
+++ b/Engine/Plugins/VirtualProduction/CaptureManager/CaptureManagerApp/Source/CaptureManagerMediaRW/Private/Readers/MHAICalibrationReader.cpp
@@ -28,10 +28,10 @@ DEFINE_LOG_CATEGORY_STATIC(LogMHAICalibrationReader, Log, All);
 		return MakeError(MoveTemp(ErrorMessage));                                      \
 	}
 
-void FMHAICalibrationReaderHelpers::RegisterReaders(FMediaRWManager* InManager)
+void FMHAICalibrationReaderHelpers::RegisterReaders(FMediaRWManager& InManager)
 {
 	TArray<FString> SupportedFormats = { TEXT("mhaical") };
-	InManager->RegisterCalibrationReader(SupportedFormats, MakeUnique<FMHAICalibrationReaderFactory>());
+	InManager.RegisterCalibrationReader(SupportedFormats, MakeUnique<FMHAICalibrationReaderFactory>());
 }
 
 TUniquePtr<ICalibrationReader> FMHAICalibrationReaderFactory::CreateCalibrationReader()
diff --git a/Engine/Plugins/VirtualProduction/CaptureManager/CaptureManagerApp/Source/CaptureManagerMediaRW/Private/Readers/MHAICalibrationReader.h b/Engine/Plugins/VirtualProduction/CaptureManager/CaptureManagerApp/Source/CaptureManagerMediaRW/Private/Readers/MHAICalibrationReader.h
--- a/Engine/Plugins/VirtualProduction/CaptureManager/CaptureManagerApp/Source/CaptureManagerMediaRW/Private/Readers/MHAICalibrationReader.h
+++ b/Engine/Plugins/VirtualProduction/CaptureManager/CaptureManagerApp/Source/CaptureManagerMediaRW/Private/Readers/MHAICalibrationReader.h
@@ -11,7 +11,7 @@ class FMHAICalibrationReaderHelpers
 {
 public:
 
-	static void RegisterReaders(class FMediaRWManager* InManager);
+	static void RegisterReaders(class FMediaRWManager& InManager);
 };
 
 class FMHAICalibrationReaderFactory : public ICalibrationReaderFactory
diff --git a/Engine/Plugins/VirtualProduction/CaptureManager/CaptureManagerApp/Source/CaptureManagerMediaRW/Private/Readers/OpenCVCalibrationReader.cpp b/Engine/Plugins/VirtualProduction/CaptureManager/CaptureManagerApp/Source/CaptureManagerMediaRW/Private/Readers/OpenCVCalibrationReader.cpp
--- a/Engine/Plugins/VirtualProduction/CaptureManager/CaptureManagerApp/Source/CaptureManagerMediaRW/Private/Readers/OpenCVCalibrationReader.cpp
+++ b/Engine/Plugins/VirtualProduction/CaptureManager/CaptureManagerApp/Source/CaptureManagerMediaRW/Private/Readers/OpenCVCalibrationReader.cpp
@@ -28,10 +28,10 @@ DEFINE_LOG_CATEGORY_STATIC(LogOpenCvCalibrationReader, Log, All);
 		return MakeError(MoveTemp(ErrorMessage));                                        \
 	}
 
-void FOpenCvCalibrationReaderHelpers::RegisterReaders(FMediaRWManager* InManager)
+void FOpenCvCalibrationReaderHelpers::RegisterReaders(FMediaRWManager& InManager)
 {
 	const TArray<FString> SupportedFormats = { TEXT("opencv") };
-	InManager->RegisterCalibrationReader(SupportedFormats, MakeUnique<FOpenCvCalibrationReaderFactory>());
+	InManager.RegisterCalibrationReader(SupportedFormats, MakeUnique<FOpenCvCalibrationReaderFactory>());
 }
 
 TUniquePtr<ICalibrationReader> FOpenCvCalibrationReaderFactory::CreateCalibrationReader()
diff --git a/Engine/Plugins/VirtualProduction/CaptureManager/CaptureManagerApp/Source/CaptureManagerMediaRW/Private/Readers/OpenCVCalibrationReader.h b/Engine/Plugins/VirtualProduction/CaptureManager/CaptureManagerApp/Source/CaptureManagerMediaRW/Private/Readers/OpenCVCalibrationReader.h
--- a/Engine/Plugins/VirtualProduction/CaptureManager/CaptureManagerApp/Source/CaptureManagerMediaRW/Private/Readers/OpenCVCalibrationReader.h
+++ b/Engine/Plugins/VirtualProduction/CaptureManager/CaptureManagerApp/Source/CaptureManagerMediaRW/Private/Readers/OpenCVCalibrationReader.h
@@ -7,12 +7,11 @@
 
 #include "Dom/JsonObject.h"
 
-// TODO: The format is specific to the HMC so it might make sense to call it differently
 class FOpenCvCalibrationReaderHelpers
 {
 public:
 
-	static void RegisterReaders(class FMediaRWManager* InManager);
+	static void RegisterReaders(class FMediaRWManager& InManager);
 };
 
 class FOpenCvCalibrationReaderFactory : public ICalibrationReaderFactory
diff --git a/Engine/Plugins/VirtualProduction/CaptureManager/CaptureManagerApp/Source/CaptureManagerMediaRW/Private/Utils/WindowsRWHelpers.cpp b/Engine/Plugins/VirtualProduction/CaptureManager/CaptureManagerApp/Source/CaptureManagerMediaRW/Private/Utils/WindowsRWHelpers.cpp
--- a/Engine/Plugins/VirtualProduction/CaptureManager/CaptureManagerApp/Source/CaptureManagerMediaRW/Private/Utils/WindowsRWHelpers.cpp
+++ b/Engine/Plugins/VirtualProduction/CaptureManager/CaptureManagerApp/Source/CaptureManagerMediaRW/Private/Utils/WindowsRWHelpers.cpp
@@ -43,17 +43,17 @@ void FWindowsRWHelpers::Deinit()
 	CoUninitialize();
 }
 
-void FWindowsRWHelpers::RegisterReaders(FMediaRWManager* InManager)
+void FWindowsRWHelpers::RegisterReaders(FMediaRWManager& InManager)
 {
 	TArray<FString> SupportedExtensions = { TEXT("mov"), TEXT("mp4") };
-	InManager->RegisterAudioReader(SupportedExtensions, MakeUnique<FWindowsReadersFactory>());
-	InManager->RegisterVideoReader(SupportedExtensions, MakeUnique<FWindowsReadersFactory>());
+	InManager.RegisterAudioReader(SupportedExtensions, MakeUnique<FWindowsReadersFactory>());
+	InManager.RegisterVideoReader(SupportedExtensions, MakeUnique<FWindowsReadersFactory>());
 }
 
-void FWindowsRWHelpers::RegisterWriters(FMediaRWManager* InManager)
+void FWindowsRWHelpers::RegisterWriters(FMediaRWManager& InManager)
 {
 	TArray<FString> SupportedExtensions = { TEXT("png"), TEXT("jpg"), TEXT("jpeg") };
-	InManager->RegisterImageWriter(SupportedExtensions, MakeUnique<FWindowsImageWriterFactory>());
+	InManager.RegisterImageWriter(SupportedExtensions, MakeUnique<FWindowsImageWriterFactory>());
 }
 
 FText FWindowsRWHelpers::CreateErrorMessage(HRESULT InResult, FText InMessage)
diff --git a/Engine/Plugins/VirtualProduction/CaptureManager/CaptureManagerApp/Source/CaptureManagerMediaRW/Private/Utils/WindowsRWHelpers.h b/Engine/Plugins/VirtualProduction/CaptureManager/CaptureManagerApp/Source/CaptureManagerMediaRW/Private/Utils/WindowsRWHelpers.h
--- a/Engine/Plugins/VirtualProduction/CaptureManager/CaptureManagerApp/Source/CaptureManagerMediaRW/Private/Utils/WindowsRWHelpers.h
+++ b/Engine/Plugins/VirtualProduction/CaptureManager/CaptureManagerApp/Source/CaptureManagerMediaRW/Private/Utils/WindowsRWHelpers.h
@@ -14,8 +14,8 @@ public:
 	static bool Init();
 	static void Deinit();
 
-	static void RegisterReaders(class FMediaRWManager* InManager);
-	static void RegisterWriters(class FMediaRWManager* InManager);
+	static void RegisterReaders(class FMediaRWManager& InManager);
+	static void RegisterWriters(class FMediaRWManager& InManager);
 
 	static FText CreateErrorMessage(HRESULT InResult, FText InMessage);
 
diff --git a/Engine/Plugins/VirtualProduction/CaptureManager/CaptureManagerApp/Source/CaptureManagerMediaRW/Private/Writers/AudioWaveMediaWriter.cpp b/Engine/Plugins/VirtualProduction/CaptureManager/CaptureManagerApp/Source/CaptureManagerMediaRW/Private/Writers/AudioWaveMediaWriter.cpp
--- a/Engine/Plugins/VirtualProduction/CaptureManager/CaptureManagerApp/Source/CaptureManagerMediaRW/Private/Writers/AudioWaveMediaWriter.cpp
+++ b/Engine/Plugins/VirtualProduction/CaptureManager/CaptureManagerApp/Source/CaptureManagerMediaRW/Private/Writers/AudioWaveMediaWriter.cpp
@@ -19,10 +19,10 @@ DEFINE_LOG_CATEGORY_STATIC(LogAudioWaveWriter, Log, All);
 		return ErrorMessage;                                                          \
 	}
 
-void FAudioWaveWriterHelpers::RegisterWriters(FMediaRWManager* InManager)
+void FAudioWaveWriterHelpers::RegisterWriters(FMediaRWManager& InManager)
 {
 	TArray<FString> SupportedFormats = { TEXT("wav") };
-	InManager->RegisterAudioWriter(SupportedFormats, MakeUnique<FAudioWaveWriterFactory>());
+	InManager.RegisterAudioWriter(SupportedFormats, MakeUnique<FAudioWaveWriterFactory>());
 }
 
 TUniquePtr<IAudioWriter> FAudioWaveWriterFactory::CreateAudioWriter()
diff --git a/Engine/Plugins/VirtualProduction/CaptureManager/CaptureManagerApp/Source/CaptureManagerMediaRW/Private/Writers/AudioWaveMediaWriter.h b/Engine/Plugins/VirtualProduction/CaptureManager/CaptureManagerApp/Source/CaptureManagerMediaRW/Private/Writers/AudioWaveMediaWriter.h
--- a/Engine/Plugins/VirtualProduction/CaptureManager/CaptureManagerApp/Source/CaptureManagerMediaRW/Private/Writers/AudioWaveMediaWriter.h
+++ b/Engine/Plugins/VirtualProduction/CaptureManager/CaptureManagerApp/Source/CaptureManagerMediaRW/Private/Writers/AudioWaveMediaWriter.h
@@ -10,7 +10,7 @@ class FAudioWaveWriterHelpers
 {
 public:
 
-	static void RegisterWriters(class FMediaRWManager* InManager);
+	static void RegisterWriters(class FMediaRWManager& InManager);
 };
 
 class FAudioWaveWriterFactory final : public IAudioWriterFactory
diff --git a/Engine/Plugins/VirtualProduction/CaptureManager/CaptureManagerApp/Source/CaptureManagerMediaRW/Private/Writers/DepthImageWriter.cpp b/Engine/Plugins/VirtualProduction/CaptureManager/CaptureManagerApp/Source/CaptureManagerMediaRW/Private/Writers/DepthImageWriter.cpp
--- a/Engine/Plugins/VirtualProduction/CaptureManager/CaptureManagerApp/Source/CaptureManagerMediaRW/Private/Writers/DepthImageWriter.cpp
+++ b/Engine/Plugins/VirtualProduction/CaptureManager/CaptureManagerApp/Source/CaptureManagerMediaRW/Private/Writers/DepthImageWriter.cpp
@@ -14,10 +14,10 @@
 
 #define LOCTEXT_NAMESPACE "DepthImageWriter"
 
-void FDepthExrImageWriterHelpers::RegisterWriters(FMediaRWManager* InManager)
+void FDepthExrImageWriterHelpers::RegisterWriters(FMediaRWManager& InManager)
 {
 	TArray<FString> SupportedFormats = { TEXT("exr") };
-	InManager->RegisterImageWriter(SupportedFormats, MakeUnique<FDepthExrImageWriterFactory>());
+	InManager.RegisterImageWriter(SupportedFormats, MakeUnique<FDepthExrImageWriterFactory>());
 }
 
 TUniquePtr<IImageWriter> FDepthExrImageWriterFactory::CreateImageWriter()
diff --git a/Engine/Plugins/VirtualProduction/CaptureManager/CaptureManagerApp/Source/CaptureManagerMediaRW/Private/Writers/DepthImageWriter.h b/Engine/Plugins/VirtualProduction/CaptureManager/CaptureManagerApp/Source/CaptureManagerMediaRW/Private/Writers/DepthImageWriter.h
--- a/Engine/Plugins/VirtualProduction/CaptureManager/CaptureManagerApp/Source/CaptureManagerMediaRW/Private/Writers/DepthImageWriter.h
+++ b/Engine/Plugins/VirtualProduction/CaptureManager/CaptureManagerApp/Source/CaptureManagerMediaRW/Private/Writers/DepthImageWriter.h
@@ -12,7 +12,7 @@ class FDepthExrImageWriterHelpers
 {
 public:
 
-	static void RegisterWriters(class FMediaRWManager* InManager);
+	static void RegisterWriters(class FMediaRWManager& InManager);
 };
 
 class FDepthExrImageWriterFactory final : public IImageWriterFactory
diff --git a/Engine/Plugins/VirtualProduction/CaptureManager/CaptureManagerApp/Source/CaptureManagerMediaRW/Private/Writers/UnrealCalibrationWriter.cpp b/Engine/Plugins/VirtualProduction/CaptureManager/CaptureManagerApp/Source/CaptureManagerMediaRW/Private/Writers/UnrealCalibrationWriter.cpp
--- a/Engine/Plugins/VirtualProduction/CaptureManager/CaptureManagerApp/Source/CaptureManagerMediaRW/Private/Writers/UnrealCalibrationWriter.cpp
+++ b/Engine/Plugins/VirtualProduction/CaptureManager/CaptureManagerApp/Source/CaptureManagerMediaRW/Private/Writers/UnrealCalibrationWriter.cpp
@@ -63,10 +63,10 @@ static void WriteOpencvDistortionModel(TSharedPtr<TJsonWriter<TCHAR>> InJsonWrit
 
 }
 
-void FUnrealCalibrationWriterHelpers::RegisterWriters(FMediaRWManager* InManager)
+void FUnrealCalibrationWriterHelpers::RegisterWriters(FMediaRWManager& InManager)
 {
 	const TArray<FString> SupportedFormats = { TEXT("unreal") };
-	InManager->RegisterCalibrationWriter(SupportedFormats, MakeUnique<FUnrealCalibrationWriterFactory>());
+	InManager.RegisterCalibrationWriter(SupportedFormats, MakeUnique<FUnrealCalibrationWriterFactory>());
 }
 
 TUniquePtr<ICalibrationWriter> FUnrealCalibrationWriterFactory::CreateCalibrationWriter()
diff --git a/Engine/Plugins/VirtualProduction/CaptureManager/CaptureManagerApp/Source/CaptureManagerMediaRW/Private/Writers/UnrealCalibrationWriter.h b/Engine/Plugins/VirtualProduction/CaptureManager/CaptureManagerApp/Source/CaptureManagerMediaRW/Private/Writers/UnrealCalibrationWriter.h
--- a/Engine/Plugins/VirtualProduction/CaptureManager/CaptureManagerApp/Source/CaptureManagerMediaRW/Private/Writers/UnrealCalibrationWriter.h
+++ b/Engine/Plugins/VirtualProduction/CaptureManager/CaptureManagerApp/Source/CaptureManagerMediaRW/Private/Writers/UnrealCalibrationWriter.h
@@ -12,7 +12,7 @@ class FUnrealCalibrationWriterHelpers
 {
 public:
 
-	static void RegisterWriters(class FMediaRWManager* InManager);
+	static void RegisterWriters(class FMediaRWManager& InManager);
 };
 
 class FUnrealCalibrationWriterFactory final : public ICalibrationWriterFactory
