Rotation when ingesting image sequence using third party encoder

diff --git a/Engine/Plugins/VirtualProduction/CaptureManager/CaptureManagerApp/Source/CaptureDataConverter/Private/Nodes/CaptureConvertDepthData.cpp b/Engine/Plugins/VirtualProduction/CaptureManager/CaptureManagerApp/Source/CaptureDataConverter/Private/Nodes/CaptureConvertDepthData.cpp
--- a/Engine/Plugins/VirtualProduction/CaptureManager/CaptureManagerApp/Source/CaptureDataConverter/Private/Nodes/CaptureConvertDepthData.cpp
+++ b/Engine/Plugins/VirtualProduction/CaptureManager/CaptureManagerApp/Source/CaptureDataConverter/Private/Nodes/CaptureConvertDepthData.cpp
@@ -5,6 +5,8 @@
 #include "Async/TaskProgress.h"
 #include "Async/StopToken.h"
 
+#include "CaptureConvertUtils.h"
+
 #include "CaptureManagerMediaRWModule.h"
 #include "CaptureCopyProgressReporter.h"
 
@@ -143,20 +145,43 @@ FCaptureConvertDepthData::FResult FCaptureConvertDepthData::CopyData()
 	const FString DepthFolderPath = FPaths::ConvertRelativePathToFull(Params.TakeOriginDirectory, Depth.Path);
 
 	IFileManager& FileManager = IFileManager::Get();
-	if (!Depth.FramesCount.IsSet() || Depth.FramesCount.GetValue() == 0)
+
+	TArray<FString> ImagesToCopy;
+	FileManager.FindFiles(ImagesToCopy, *DepthFolderPath, Depth.Format.IsEmpty() ? nullptr : *Depth.Format);
+
+	IImageWrapperModule& ImageWrapperModule = FModuleManager::Get().LoadModuleChecked<IImageWrapperModule>(TEXT("ImageWrapper"));
+	ImagesToCopy.RemoveAllSwap([&ImageWrapperModule](const FString& InFilePath)
 	{
-		TArray<FString> FoundFiles;
-		FileManager.FindFiles(FoundFiles, *DepthFolderPath, Depth.Format.IsEmpty() ? nullptr : *Depth.Format);
-		Depth.FramesCount = FoundFiles.Num();
+		return ImageWrapperModule.GetImageFormatFromExtension(*InFilePath) != EImageFormat::EXR;
+	});
 
-		if (!Depth.FramesCount.IsSet() || Depth.FramesCount.GetValue() == 0)
-		{
-			FText ExtensionPostfix = Depth.Format.IsEmpty() ? FText() : FText::Format(LOCTEXT("CaptureConvertDepth_Extension", " with specified extension .{0}"), FText::FromString(Depth.Format));
-			FText Message = FText::Format(LOCTEXT("CaptureConvertDepth_EmptyData", "Copy depth data failed. No depth data found at {0}{1}"), FText::FromString(DepthFolderPath), ExtensionPostfix);
-			UE_LOG(LogCaptureConvertDepthData, Error, TEXT("%s"), *Message.ToString());
+	if (ImagesToCopy.IsEmpty())
+	{
+		FText ExtensionPostfix = Depth.Format.IsEmpty() ? FText() : FText::Format(LOCTEXT("CaptureConvertDepth_Extension", " with specified extension .{0}"), FText::FromString(Depth.Format));
+		FText Message = FText::Format(LOCTEXT("CaptureConvertDepth_EmptyData", "Copy depth data failed. No depth data found at {0}{1}"), FText::FromString(DepthFolderPath), ExtensionPostfix);
+		UE_LOG(LogCaptureConvertDepthData, Error, TEXT("%s"), *Message.ToString());
 
-			return MakeError(MoveTemp(Message));
+		return MakeError(MoveTemp(Message));
+	}
+
+	ImagesToCopy.Sort([](const FString& InLeft, const FString& InRight)
+	{
+		FString Dummy;
+		FString LeftDigits, RightDigits;
+		bool bSuccess = ExtractInfoFromFileName(InLeft, Dummy, LeftDigits, Dummy);
+		bSuccess = bSuccess && ExtractInfoFromFileName(InRight, Dummy, RightDigits, Dummy);
+
+		if (bSuccess)
+		{
+			return FCString::Atoi(*LeftDigits) < FCString::Atoi(*RightDigits);
 		}
+
+		return InLeft < InRight;
+	});
+
+	if (!Depth.FramesCount.IsSet() || Depth.FramesCount.GetValue() == 0)
+	{
+		Depth.FramesCount = ImagesToCopy.Num();
 	}
 
 	FTaskProgress::FTask Task = Params.TaskProgress->StartTask();
@@ -166,57 +191,32 @@ FCaptureConvertDepthData::FResult FCaptureConvertDepthData::CopyData()
 	}));
 
 	FString DestinationDirectory = OutputDirectory / Depth.Name;
-	FCaptureConvertDepthData::FResult Result = MakeValue();
-	int32 NumberOfCopiedFiles = 0;
 
-	IImageWrapperModule& ImageWrapperModule = FModuleManager::Get().LoadModuleChecked<IImageWrapperModule>(TEXT("ImageWrapper"));
-	FileManager.IterateDirectoryRecursively(*DepthFolderPath, [this, &TaskProgress, &Result, DestinationDirectory, &FileManager, &ImageWrapperModule, &NumberOfCopiedFiles](const TCHAR* InFileName, bool bInIsDirectory)
+	for (const FString& FilePath : ImagesToCopy)
 	{
-		if (bInIsDirectory)
-		{
-			return true;
-		}
+		FTaskProgress::FTask InnerTask = TaskProgress->StartTask();
 
-		if (ImageWrapperModule.GetImageFormatFromExtension(InFileName) != EImageFormat::EXR)
-		{
-			UE_LOG(LogCaptureConvertDepthData, Warning, TEXT("Image file format is unsupported %s"), InFileName);
-			return true;
-		}
-		
-		FTaskProgress::FTask Task = TaskProgress->StartTask();
+		FCopyProgressReporter ProgressReporter(InnerTask, Params.StopToken);
 
-		FCopyProgressReporter ProgressReporter(Task, Params.StopToken);
+		FString Destination = DestinationDirectory / FilePath;
+		FString Source = DepthFolderPath / FilePath;
 
-		FString Destination = DestinationDirectory / FPaths::GetCleanFilename(InFileName);
-
-		uint32 CopyResult = FileManager.Copy(*Destination, InFileName, true, true, false, &ProgressReporter);
+		uint32 CopyResult = FileManager.Copy(*Destination, *Source, true, true, false, &ProgressReporter);
 
 		if (CopyResult == COPY_Fail)
 		{
 			FText Message = LOCTEXT("CaptureConvertDepthData_CopyFailed", "Failed to copy the file");
-			Result = MakeError(MoveTemp(Message));
-			return false;
+			return MakeError(MoveTemp(Message));
 		}
 
 		if (CopyResult == COPY_Canceled)
 		{
 			FText Message = LOCTEXT("CaptureConvertDepthData_AbortedByUser", "Aborted by user");
-			Result = MakeError(MoveTemp(Message));
-			return false;
+			return MakeError(MoveTemp(Message));
 		}
-
-		++NumberOfCopiedFiles;
-
-		return true;
-	});
-
-	if (NumberOfCopiedFiles == 0)
-	{
-		FText Message = FText::Format(LOCTEXT("CaptureConvertDepthData_NoFiles", "Failed to copy from {0} to {1}"), FText::FromString(DepthFolderPath), FText::FromString(DestinationDirectory));
-		return MakeError(MoveTemp(Message));
 	}
 
-	return Result;
+	return MakeValue();
 }
 
 FCaptureConvertDepthData::FWriteResult FCaptureConvertDepthData::OnWrite(FWritingContext InContext)
diff --git a/Engine/Plugins/VirtualProduction/CaptureManager/CaptureManagerApp/Source/CaptureDataConverter/Private/Nodes/CaptureConvertUtils.cpp b/Engine/Plugins/VirtualProduction/CaptureManager/CaptureManagerApp/Source/CaptureDataConverter/Private/Nodes/CaptureConvertUtils.cpp
new file mode 100644
--- /dev/null
+++ b/Engine/Plugins/VirtualProduction/CaptureManager/CaptureManagerApp/Source/CaptureDataConverter/Private/Nodes/CaptureConvertUtils.cpp
@@ -0,0 +1,73 @@
+// Copyright Epic Games, Inc. All Rights Reserved.
+
+#include "CaptureConvertUtils.h"
+
+#include "HAL/FileManager.h"
+#include "Modules/ModuleManager.h"
+#include "IImageWrapperModule.h"
+
+#include "Misc/Paths.h"
+#include "Internationalization/Regex.h"
+
+namespace UE::CaptureManager
+{
+
+bool ExtractInfoFromFileName(const FString& InFileName, FString& OutPrefix, FString& OutDigits, FString& OutExtension)
+{
+	FRegexPattern Pattern(TEXT("^(.*?)(\\d+)\\.(\\w+)$"));
+	FRegexMatcher Matcher(Pattern, InFileName);
+
+	if (Matcher.FindNext())
+	{
+		OutPrefix = Matcher.GetCaptureGroup(1);
+		OutDigits = Matcher.GetCaptureGroup(2);
+		OutExtension = Matcher.GetCaptureGroup(3);
+
+		return true;
+	}
+
+	return false;
+}
+
+FString GetFileFormat(const FString& InFileName)
+{
+	FString Prefix;
+	FString Digits;
+	FString Extension;
+	if (ExtractInfoFromFileName(InFileName, Prefix, Digits, Extension))
+	{
+		return FString::Format(TEXT("{0}%0{1}d.{2}"), { Prefix, FString::FromInt(Digits.Len()), Extension });
+	}
+
+	return FString();
+}
+
+FString GetFileNameFormat(const FString& InDirectory)
+{
+	IFileManager& FileManager = IFileManager::Get();
+	IImageWrapperModule& ImageWrapperModule =
+		FModuleManager::Get().LoadModuleChecked<IImageWrapperModule>(TEXT("ImageWrapper"));
+	FString FileName;
+	FileManager.IterateDirectory(*InDirectory, [&FileName, &ImageWrapperModule](const TCHAR* InFilenameOrDirectory, bool bInIsDirectory)
+	{
+		if (bInIsDirectory)
+		{
+			return true;
+		}
+
+		FString Extension = FPaths::GetExtension(InFilenameOrDirectory, false);
+		if (ImageWrapperModule.GetImageFormatFromExtension(*Extension) != EImageFormat::Invalid)
+		{
+			// Found the correct extension and return false to exit the iteration
+			FileName = FPaths::GetCleanFilename(InFilenameOrDirectory);
+
+			return false;
+		}
+
+		return true;
+	});
+
+	return GetFileFormat(FileName);
+}
+
+}
diff --git a/Engine/Plugins/VirtualProduction/CaptureManager/CaptureManagerApp/Source/CaptureDataConverter/Private/Nodes/CaptureConvertUtils.h b/Engine/Plugins/VirtualProduction/CaptureManager/CaptureManagerApp/Source/CaptureDataConverter/Private/Nodes/CaptureConvertUtils.h
new file mode 100644
--- /dev/null
+++ b/Engine/Plugins/VirtualProduction/CaptureManager/CaptureManagerApp/Source/CaptureDataConverter/Private/Nodes/CaptureConvertUtils.h
@@ -0,0 +1,15 @@
+// Copyright Epic Games, Inc. All Rights Reserved.
+
+#pragma once
+
+#include "Containers/UnrealString.h"
+
+namespace UE::CaptureManager
+{
+
+bool ExtractInfoFromFileName(const FString& InFileName, FString& OutPrefix, FString& OutDigits, FString& OutExtension);
+
+FString GetFileFormat(const FString& InFileName);
+FString GetFileNameFormat(const FString& InDirectory);
+
+}
\ No newline at end of file
diff --git a/Engine/Plugins/VirtualProduction/CaptureManager/CaptureManagerApp/Source/CaptureDataConverter/Private/Nodes/CaptureConvertVideoData.cpp b/Engine/Plugins/VirtualProduction/CaptureManager/CaptureManagerApp/Source/CaptureDataConverter/Private/Nodes/CaptureConvertVideoData.cpp
--- a/Engine/Plugins/VirtualProduction/CaptureManager/CaptureManagerApp/Source/CaptureDataConverter/Private/Nodes/CaptureConvertVideoData.cpp
+++ b/Engine/Plugins/VirtualProduction/CaptureManager/CaptureManagerApp/Source/CaptureDataConverter/Private/Nodes/CaptureConvertVideoData.cpp
@@ -8,6 +8,8 @@
 #include "CaptureManagerMediaRWModule.h"
 #include "CaptureCopyProgressReporter.h"
 
+#include "CaptureConvertUtils.h"
+
 #include "IImageWrapperModule.h"
 
 #define LOCTEXT_NAMESPACE "CaptureConvertVideoData"
@@ -237,21 +239,43 @@ FCaptureConvertVideoData::FResult FCaptureConvertVideoData::CopyData()
 	const FString VideoFolderPath = FPaths::ConvertRelativePathToFull(Params.TakeOriginDirectory, Video.Path);
 
 	IFileManager& FileManager = IFileManager::Get();
-	if (!Video.FramesCount.IsSet() || Video.FramesCount.GetValue() == 0)
+
+	TArray<FString> ImagesToCopy;
+	FileManager.FindFiles(ImagesToCopy, *VideoFolderPath, Video.Format.IsEmpty() ? nullptr : *Video.Format);
+
+	IImageWrapperModule& ImageWrapperModule = FModuleManager::Get().LoadModuleChecked<IImageWrapperModule>(TEXT("ImageWrapper"));
+	ImagesToCopy.RemoveAllSwap([&ImageWrapperModule](const FString& InFilePath)
 	{
-		TArray<FString> FoundFiles;
-		FileManager.FindFiles(FoundFiles, *VideoFolderPath, Video.Format.IsEmpty() ? nullptr : *Video.Format);
-		
-		if (FoundFiles.IsEmpty())
-		{
-			FText ExtensionPostfix = Video.Format.IsEmpty() ? FText() : FText::Format(LOCTEXT("CaptureConvertVideo_Extension", " with specified extension .{0}"), FText::FromString(Video.Format));
-			FText Message = FText::Format(LOCTEXT("CaptureConvertVideo_EmptyData", "Copy image data failed. No image data found at {0}{1}"), FText::FromString(VideoFolderPath), ExtensionPostfix);
-			UE_LOG(LogCaptureConvertVideoData, Error, TEXT("%s"), *Message.ToString());
+		return ImageWrapperModule.GetImageFormatFromExtension(*InFilePath) == EImageFormat::Invalid;
+	});
 
-			return MakeError(MoveTemp(Message));
+	if (ImagesToCopy.IsEmpty())
+	{
+		FText ExtensionPostfix = Video.Format.IsEmpty() ? FText() : FText::Format(LOCTEXT("CaptureConvertVideo_Extension", " with specified extension .{0}"), FText::FromString(Video.Format));
+		FText Message = FText::Format(LOCTEXT("CaptureConvertVideo_EmptyData", "Copy image data failed. No image data found at {0}{1}"), FText::FromString(VideoFolderPath), ExtensionPostfix);
+		UE_LOG(LogCaptureConvertVideoData, Error, TEXT("%s"), *Message.ToString());
+
+		return MakeError(MoveTemp(Message));
+	}
+
+	ImagesToCopy.Sort([](const FString& InLeft, const FString& InRight)
+	{
+		FString Dummy;
+		FString LeftDigits, RightDigits;
+		bool bSuccess = ExtractInfoFromFileName(InLeft, Dummy, LeftDigits, Dummy);
+		bSuccess = bSuccess && ExtractInfoFromFileName(InRight, Dummy, RightDigits, Dummy);
+
+		if (bSuccess)
+		{
+			return FCString::Atoi(*LeftDigits) < FCString::Atoi(*RightDigits);
 		}
 
-		Video.FramesCount = FoundFiles.Num();
+		return InLeft < InRight;
+	});
+
+	if (!Video.FramesCount.IsSet() || Video.FramesCount.GetValue() == 0)
+	{
+		Video.FramesCount = ImagesToCopy.Num();
 	}
 
 	FTaskProgress::FTask Task = Params.TaskProgress->StartTask();
@@ -262,57 +286,35 @@ FCaptureConvertVideoData::FResult FCaptureConvertVideoData::CopyData()
 
 	FString DestinationDirectory = OutputDirectory / Video.Name;
 
-	FCaptureConvertVideoData::FResult Result = MakeValue();
-	int32 NumberOfCopiedFiles = 0;
-
-	IImageWrapperModule& ImageWrapperModule = FModuleManager::Get().LoadModuleChecked<IImageWrapperModule>(TEXT("ImageWrapper"));
-	FileManager.IterateDirectoryRecursively(*VideoFolderPath, [this, &TaskProgress, &Result, DestinationDirectory, &FileManager, &ImageWrapperModule, &NumberOfCopiedFiles, VideoFolderPath](const TCHAR* InFileName, bool bInIsDirectory)
+	for (const FString& FilePath : ImagesToCopy)
 	{
-		if (bInIsDirectory)
-		{
-			return true;
-		}
+		FTaskProgress::FTask InnerTask = TaskProgress->StartTask();
 
-		if (ImageWrapperModule.GetImageFormatFromExtension(InFileName) == EImageFormat::Invalid)
-		{
-			UE_LOG(LogCaptureConvertVideoData, Warning, TEXT("Image file format is unsupported %s"), InFileName);
-			return true;
-		}
-		
-		FTaskProgress::FTask Task = TaskProgress->StartTask();
-
-		FCopyProgressReporter ProgressReporter(Task, Params.StopToken);
+		FCopyProgressReporter ProgressReporter(InnerTask, Params.StopToken);
 
-		FString Destination = DestinationDirectory / FPaths::GetCleanFilename(InFileName);
+		FString Destination = DestinationDirectory / FilePath;
+		FString Source = VideoFolderPath / FilePath;
 
-		uint32 CopyResult = FileManager.Copy(*Destination, InFileName, true, true, false, &ProgressReporter);
+		constexpr bool bReplace = true;
+		constexpr bool bEvenIfReadOnly = true;
+		constexpr bool bAttributes = false;
+		uint32 CopyResult = FileManager.Copy(*Destination, *Source, bReplace, bEvenIfReadOnly, bAttributes, &ProgressReporter);
 
 		if (CopyResult == COPY_Fail)
 		{
-			FText Message = FText::Format(LOCTEXT("CaptureConvertVideoData_CopyFailed", "Failed to copy file {0} from {1} to {2}"), FText::FromString(InFileName), FText::FromString(VideoFolderPath), FText::FromString(DestinationDirectory));
-			Result = MakeError(MoveTemp(Message));
-			return false;
+			FText Message = FText::Format(LOCTEXT("CaptureConvertVideoData_CopyFailed", "Failed to copy file {0} from {1} to {2}"), 
+										  FText::FromString(FilePath), FText::FromString(VideoFolderPath), FText::FromString(DestinationDirectory));
+			return MakeError(MoveTemp(Message));
 		}
 
 		if (CopyResult == COPY_Canceled)
 		{
 			FText Message = LOCTEXT("CaptureConvertVideoData_AbortedByUser", "Image data copy aborted by user");
-			Result = MakeError(MoveTemp(Message));
-			return false;
+			return MakeError(MoveTemp(Message));
 		}
-
-		++NumberOfCopiedFiles;
-
-		return true;
-	});
-
-	if (NumberOfCopiedFiles == 0)
-	{
-		FText Message = FText::Format(LOCTEXT("CaptureConvertVideoData_NoFiles", "Failed to copy from {0} to {1}"), FText::FromString(VideoFolderPath), FText::FromString(DestinationDirectory));
-		return MakeError(MoveTemp(Message));
 	}
 
-	return Result;
+	return MakeValue();
 }
 
 FCaptureConvertVideoData::FWriteResult FCaptureConvertVideoData::OnWrite(FWritingContext InContext)
diff --git a/Engine/Plugins/VirtualProduction/CaptureManager/CaptureManagerApp/Source/CaptureDataConverter/Private/Nodes/ThirdPartyEncoder/CaptureConvertVideoDataThirdParty.cpp b/Engine/Plugins/VirtualProduction/CaptureManager/CaptureManagerApp/Source/CaptureDataConverter/Private/Nodes/ThirdPartyEncoder/CaptureConvertVideoDataThirdParty.cpp
--- a/Engine/Plugins/VirtualProduction/CaptureManager/CaptureManagerApp/Source/CaptureDataConverter/Private/Nodes/ThirdPartyEncoder/CaptureConvertVideoDataThirdParty.cpp
+++ b/Engine/Plugins/VirtualProduction/CaptureManager/CaptureManagerApp/Source/CaptureDataConverter/Private/Nodes/ThirdPartyEncoder/CaptureConvertVideoDataThirdParty.cpp
@@ -8,6 +8,8 @@
 #include "Async/StopToken.h"
 #include "Async/TaskProgress.h"
 
+#include "Nodes/CaptureConvertUtils.h"
+
 #include "CaptureManagerMediaRWModule.h"
 #include "Nodes/CaptureCopyProgressReporter.h"
 #include "CaptureThirdPartyNodeUtils.h"
@@ -21,6 +23,8 @@
 
 #include "IImageWrapperModule.h"
 
+#include "Internationalization/Regex.h"
+
 #define LOCTEXT_NAMESPACE "CaptureConvertVideoDataTP"
 
 DEFINE_LOG_CATEGORY_STATIC(LogCaptureConvertVideoDataThirdParty, Log, All);
@@ -145,12 +149,12 @@ FCaptureConvertVideoDataThirdParty::FResult FCaptureConvertVideoDataThirdParty::
 		FTakeMetadataPathUtils::ValidatePathType(Video.Path, Video.PathType.GetValue());
 	}
 
-	if (Video.PathType.IsSet() && Video.PathType.GetValue() == FTakeMetadata::FVideo::EPathType::File)
+	if (ShouldCopy())
 	{
-		return ConvertData();
+		return CopyData();
 	}
-	
-	return CopyData();
+
+	return ConvertData();
 }
 
 FCaptureConvertVideoDataThirdParty::FResult FCaptureConvertVideoDataThirdParty::ConvertData()
@@ -159,7 +163,7 @@ FCaptureConvertVideoDataThirdParty::FResult FCaptureConvertVideoDataThirdParty::
 
 	FTaskProgress::FTask Task = Params.TaskProgress->StartTask();
 
-	const FString VideoFilePath = FPaths::ConvertRelativePathToFull(Params.TakeOriginDirectory, Video.Path);
+	FString VideoFilePath = FPaths::ConvertRelativePathToFull(Params.TakeOriginDirectory, Video.Path);
 	FString DestinationDirectory = OutputDirectory / Video.Name;
 	
 	if (!FPaths::DirectoryExists(DestinationDirectory))
@@ -167,14 +171,32 @@ FCaptureConvertVideoDataThirdParty::FResult FCaptureConvertVideoDataThirdParty::
 		bool bSuccess = IFileManager::Get().MakeDirectory(*DestinationDirectory, true);
 		if (!bSuccess)
 		{
-			UE_LOG(LogCaptureConvertVideoDataThirdParty, Error, TEXT("Failed to create directory: %s"), *DestinationDirectory);
 			FText Message = FText::Format(
 				LOCTEXT("CaptureConvertVideoNodeTP_FailedToCreateDirectory", "Failed to create encoder output directory: {0}"),
 				FText::FromString(DestinationDirectory)
 			);
 
+			UE_LOG(LogCaptureConvertVideoDataThirdParty, Error, TEXT("%s"), *Message.ToString());
+			return MakeError(MoveTemp(Message));
+		}
+	}
+
+	if (Video.PathType.IsSet() && Video.PathType.GetValue() == FTakeMetadata::FVideo::EPathType::Folder)
+	{
+		FString FilePattern = GetFileNameFormat(VideoFilePath);
+		if (FilePattern.IsEmpty())
+		{
+			FText Message = FText::Format(
+				LOCTEXT("CaptureConvertVideoNodeTP_FailedToDetectFilePattern", "Failed to detect the file pattern for directory: {0}"),
+				FText::FromString(VideoFilePath)
+			);
+
+			UE_LOG(LogCaptureConvertVideoDataThirdParty, Error, TEXT("%s"), *Message.ToString());
+			
 			return MakeError(MoveTemp(Message));
 		}
+
+		VideoFilePath = FPaths::Combine(VideoFilePath, FilePattern);
 	}
 
 	const FString ImageFileName = FString::Format(TEXT("{0}_%06d.{1}"), { VideoParams.ImageFileName, VideoParams.Format });
@@ -285,6 +307,75 @@ FCaptureConvertVideoDataThirdParty::FResult FCaptureConvertVideoDataThirdParty::
 	return MakeValue();
 }
 
+bool FCaptureConvertVideoDataThirdParty::ShouldCopy() const
+{
+	FCaptureManagerMediaRWModule& MediaRWModule =
+		FModuleManager::LoadModuleChecked<FCaptureManagerMediaRWModule>("CaptureManagerMediaRW");
+
+	if (!Video.PathType.IsSet() || Video.PathType.GetValue() != FTakeMetadata::FVideo::EPathType::Folder)
+	{
+		return false;
+	}
+
+	TUniquePtr<IVideoReader> VideoReader = MediaRWModule.Get().CreateVideoReaderByFormat(TEXT("image_sequence"));
+
+	if (!VideoReader)
+	{
+		return false;
+	}
+
+	const FString VideoFilePath = FPaths::ConvertRelativePathToFull(Params.TakeOriginDirectory, Video.Path);
+	TOptional<FText> Error = VideoReader->Open(VideoFilePath);
+
+	if (Error.IsSet())
+	{
+		return false;
+	}
+
+	TValueOrError<TUniquePtr<UE::CaptureManager::FMediaTextureSample>, FText> ReadResult = VideoReader->Next();
+	if (!ReadResult.IsValid())
+	{
+		return false;
+	}
+
+	TUniquePtr<UE::CaptureManager::FMediaTextureSample> Sample = ReadResult.StealValue();
+
+	if (!Sample)
+	{
+		return false;
+	}
+
+	IFileManager& FileManager = IFileManager::Get();
+	IImageWrapperModule& ImageWrapperModule =
+		FModuleManager::Get().LoadModuleChecked<IImageWrapperModule>(TEXT("ImageWrapper"));
+	FString Extension;
+	FileManager.IterateDirectory(*VideoFilePath, [&Extension, &ImageWrapperModule](const TCHAR* InFilenameOrDirectory, bool bInIsDirectory)
+	{
+		if (bInIsDirectory)
+		{
+			return true;
+		}
+
+		Extension = FPaths::GetExtension(InFilenameOrDirectory, false);
+		if (ImageWrapperModule.GetImageFormatFromExtension(*Extension) != EImageFormat::Invalid)
+		{
+			// Found the correct extension and return false to exit the iteration
+			return false;
+		}
+
+		return true;
+	});
+
+	if (Sample->CurrentFormat == VideoParams.OutputPixelFormat &&
+		VideoParams.Rotation == EMediaOrientation::Original &&
+		VideoParams.Format == Extension)
+	{
+		return true;
+	}
+
+	return false;
+}
+
 FCaptureConvertVideoDataThirdParty::FResult FCaptureConvertVideoDataThirdParty::CopyData()
 {
 	using namespace UE::CaptureManager;
@@ -292,21 +383,43 @@ FCaptureConvertVideoDataThirdParty::FResult FCaptureConvertVideoDataThirdParty::
 	const FString VideoFolderPath = FPaths::ConvertRelativePathToFull(Params.TakeOriginDirectory, Video.Path);
 
 	IFileManager& FileManager = IFileManager::Get();
-	if (!Video.FramesCount.IsSet() || Video.FramesCount.GetValue() == 0)
+
+	TArray<FString> ImagesToCopy;
+	FileManager.FindFiles(ImagesToCopy, *VideoFolderPath, Video.Format.IsEmpty() ? nullptr : *Video.Format);
+
+	IImageWrapperModule& ImageWrapperModule = FModuleManager::Get().LoadModuleChecked<IImageWrapperModule>(TEXT("ImageWrapper"));
+	ImagesToCopy.RemoveAllSwap([&ImageWrapperModule](const FString& InFilePath)
 	{
-		TArray<FString> FoundFiles;
-		FileManager.FindFiles(FoundFiles, *VideoFolderPath, Video.Format.IsEmpty() ? nullptr : *Video.Format);
+		return ImageWrapperModule.GetImageFormatFromExtension(*InFilePath) == EImageFormat::Invalid;
+	});
 
-		if (FoundFiles.IsEmpty())
-		{
-			FText ExtensionPostfix = Video.Format.IsEmpty() ? FText() : FText::Format(LOCTEXT("CaptureConvertVideo_Extension", " with specified extension .{0}"), FText::FromString(Video.Format));
-			FText Message = FText::Format(LOCTEXT("CaptureConvertVideoNodeTP_EmptyData", "Copy image data failed. No image data found at {0}{1}"), FText::FromString(VideoFolderPath), ExtensionPostfix);
-			UE_LOG(LogCaptureConvertVideoDataThirdParty, Error, TEXT("%s"), *Message.ToString());
+	if (ImagesToCopy.IsEmpty())
+	{
+		FText ExtensionPostfix = Video.Format.IsEmpty() ? FText() : FText::Format(LOCTEXT("CaptureConvertVideoTP_Extension", " with specified extension .{0}"), FText::FromString(Video.Format));
+		FText Message = FText::Format(LOCTEXT("CaptureConvertVideoNodeTP_EmptyData", "Copy image data failed. No image data found at {0}{1}"), FText::FromString(VideoFolderPath), ExtensionPostfix);
+		UE_LOG(LogCaptureConvertVideoDataThirdParty, Error, TEXT("%s"), *Message.ToString());
 
-			return MakeError(MoveTemp(Message));
+		return MakeError(MoveTemp(Message));
+	}
+
+	ImagesToCopy.Sort([](const FString& InLeft, const FString& InRight)
+	{
+		FString Dummy;
+		FString LeftDigits, RightDigits;
+		bool bSuccess = ExtractInfoFromFileName(InLeft, Dummy, LeftDigits, Dummy);
+		bSuccess = bSuccess && ExtractInfoFromFileName(InRight, Dummy, RightDigits, Dummy);
+
+		if (bSuccess)
+		{
+			return FCString::Atoi(*LeftDigits) < FCString::Atoi(*RightDigits);
 		}
 
-		Video.FramesCount = FoundFiles.Num();
+		return InLeft < InRight;
+	});
+
+	if (!Video.FramesCount.IsSet() || Video.FramesCount.GetValue() == 0)
+	{
+		Video.FramesCount = ImagesToCopy.Num();
 	}
 
 	FTaskProgress::FTask Task = Params.TaskProgress->StartTask();
@@ -317,61 +430,35 @@ FCaptureConvertVideoDataThirdParty::FResult FCaptureConvertVideoDataThirdParty::
 
 	FString DestinationDirectory = OutputDirectory / Video.Name;
 
-	FCaptureConvertVideoDataThirdParty::FResult Result = MakeValue();
-	int32 NumberOfCopiedFiles = 0;
-
-	IImageWrapperModule& ImageWrapperModule = FModuleManager::Get().LoadModuleChecked<IImageWrapperModule>(TEXT("ImageWrapper"));
-	FileManager.IterateDirectoryRecursively(*VideoFolderPath, [this, &TaskProgress, &Result, DestinationDirectory, &FileManager, &ImageWrapperModule, &NumberOfCopiedFiles](const TCHAR* InFileName, bool bInIsDirectory)
+	for (const FString& FilePath : ImagesToCopy)
 	{
-		if (bInIsDirectory)
-		{
-			return true;
-		}
+		FTaskProgress::FTask InnerTask = TaskProgress->StartTask();
 
-		if (ImageWrapperModule.GetImageFormatFromExtension(InFileName) == EImageFormat::Invalid)
-		{
-			UE_LOG(LogCaptureConvertVideoDataThirdParty, Warning, TEXT("Image file format is unsupported %s"), InFileName);
-			return true;
-		}
-		
-		FTaskProgress::FTask Task = TaskProgress->StartTask();
-
-		FCopyProgressReporter ProgressReporter(Task, Params.StopToken);
+		FCopyProgressReporter ProgressReporter(InnerTask, Params.StopToken);
 
-		FString Destination = DestinationDirectory / FPaths::GetCleanFilename(InFileName);
+		FString Destination = DestinationDirectory / FilePath;
+		FString Source = VideoFolderPath / FilePath;
 
 		constexpr bool bReplace = true;
 		constexpr bool bEvenIfReadOnly = true;
 		constexpr bool bAttributes = false;
-		uint32 CopyResult = FileManager.Copy(*Destination, InFileName, bReplace, bEvenIfReadOnly, bAttributes, &ProgressReporter);
+		uint32 CopyResult = FileManager.Copy(*Destination, *Source, bReplace, bEvenIfReadOnly, bAttributes, &ProgressReporter);
 
 		if (CopyResult == COPY_Fail)
 		{
-			FText Message = FText::Format(LOCTEXT("CaptureConvertVideoNodeTP_CopyFailed", "Failed to copy the video file {0}"),
-										  FText::FromString(InFileName));
-			Result = MakeError(MoveTemp(Message));
-			return false;
+			FText Message = FText::Format(LOCTEXT("CaptureConvertVideoNodeTP_CopyFailed", "Failed to copy file {0} from {1} to {2}"),
+										  FText::FromString(FilePath), FText::FromString(VideoFolderPath), FText::FromString(DestinationDirectory));
+			return MakeError(MoveTemp(Message));
 		}
 
 		if (CopyResult == COPY_Canceled)
 		{
-			FText Message = LOCTEXT("CaptureConvertVideoNodeTP_AbortedByUser", "Aborted by user");
-			Result = MakeError(MoveTemp(Message));
-			return false;
+			FText Message = LOCTEXT("CaptureConvertVideoNodeTP_AbortedByUser", "Image data copy aborted by user");
+			return MakeError(MoveTemp(Message));
 		}
-
-		++NumberOfCopiedFiles;
-
-		return true;
-	});
-
-	if (NumberOfCopiedFiles == 0)
-	{
-		FText Message = FText::Format(LOCTEXT("CaptureConvertVideoNodeTP_NoFiles", "Failed to copy from {0} to {1}"), FText::FromString(VideoFolderPath), FText::FromString(DestinationDirectory));
-		return MakeError(MoveTemp(Message));
 	}
 
-	return Result;
+	return MakeValue();
 }
 
 #undef LOCTEXT_NAMESPACE
diff --git a/Engine/Plugins/VirtualProduction/CaptureManager/CaptureManagerApp/Source/CaptureDataConverter/Private/Nodes/ThirdPartyEncoder/CaptureConvertVideoDataThirdParty.h b/Engine/Plugins/VirtualProduction/CaptureManager/CaptureManagerApp/Source/CaptureDataConverter/Private/Nodes/ThirdPartyEncoder/CaptureConvertVideoDataThirdParty.h
--- a/Engine/Plugins/VirtualProduction/CaptureManager/CaptureManagerApp/Source/CaptureDataConverter/Private/Nodes/ThirdPartyEncoder/CaptureConvertVideoDataThirdParty.h
+++ b/Engine/Plugins/VirtualProduction/CaptureManager/CaptureManagerApp/Source/CaptureDataConverter/Private/Nodes/ThirdPartyEncoder/CaptureConvertVideoDataThirdParty.h
@@ -24,6 +24,7 @@ private:
 	virtual FResult Run() override;
 
 	FResult ConvertData();
+	bool ShouldCopy() const;
 	FResult CopyData();
 
 	FCaptureThirdPartyNodeParams ThirdPartyEncoder;
diff --git a/Engine/Plugins/VirtualProduction/CaptureManager/CaptureManagerApp/Source/CaptureManagerMediaRW/Private/Readers/ImageSequenceReader.cpp b/Engine/Plugins/VirtualProduction/CaptureManager/CaptureManagerApp/Source/CaptureManagerMediaRW/Private/Readers/ImageSequenceReader.cpp
--- a/Engine/Plugins/VirtualProduction/CaptureManager/CaptureManagerApp/Source/CaptureManagerMediaRW/Private/Readers/ImageSequenceReader.cpp
+++ b/Engine/Plugins/VirtualProduction/CaptureManager/CaptureManagerApp/Source/CaptureManagerMediaRW/Private/Readers/ImageSequenceReader.cpp
@@ -8,11 +8,30 @@
 #include "IImageWrapperModule.h"
 #include "Modules/ModuleManager.h"
 
+#include "Internationalization/Regex.h"
+
 #define LOCTEXT_NAMESPACE "ImageSequenceReader"
 
 namespace UE::CaptureManager::Private
 {
 
+bool ExtractInfoFromFileName(const FString& InFileName, FString& OutPrefix, FString& OutDigits, FString& OutExtension)
+{
+	FRegexPattern Pattern(TEXT("^(.*?)(\\d+)\\.(\\w+)$"));
+	FRegexMatcher Matcher(Pattern, InFileName);
+
+	if (Matcher.FindNext())
+	{
+		OutPrefix = Matcher.GetCaptureGroup(1);
+		OutDigits = Matcher.GetCaptureGroup(2);
+		OutExtension = Matcher.GetCaptureGroup(3);
+
+		return true;
+	}
+
+	return false;
+}
+
 static TArray<FString> GetImageSequenceFilesFromPath(const FString& InFullSequencePath)
 {
 	TArray<FString> ImageFiles;
@@ -121,7 +140,20 @@ TOptional<FText> FImageSequenceReader::Open(const FString& InDirectoryPath)
 		return Error;
 	}
 
-	ImagePaths.Sort();
+	ImagePaths.Sort([](const FString& InLeft, const FString& InRight)
+	{
+		FString Dummy;
+		FString LeftDigits, RightDigits;
+		bool bSuccess = UE::CaptureManager::Private::ExtractInfoFromFileName(InLeft, Dummy, LeftDigits, Dummy);
+		bSuccess = bSuccess && UE::CaptureManager::Private::ExtractInfoFromFileName(InRight, Dummy, RightDigits, Dummy);
+
+		if (bSuccess)
+		{
+			return FCString::Atoi(*LeftDigits) < FCString::Atoi(*RightDigits);
+		}
+
+		return InLeft < InRight;
+	});
 
 	check(ImagePaths.IsValidIndex(0));
 	UE::CaptureManager::Private::GetImageDimensionsFromPath(ImagePaths[0], Dimensions);
