Text conflict from LOCTEXT macro for namespace 'CaptureConvertVideoDataTP'

diff --git a/Engine/Plugins/VirtualProduction/CaptureManager/CaptureManagerApp/Source/CaptureDataConverter/Private/Nodes/ThirdPartyEncoder/CaptureConvertVideoDataThirdParty.cpp b/Engine/Plugins/VirtualProduction/CaptureManager/CaptureManagerApp/Source/CaptureDataConverter/Private/Nodes/ThirdPartyEncoder/CaptureConvertVideoDataThirdParty.cpp
--- a/Engine/Plugins/VirtualProduction/CaptureManager/CaptureManagerApp/Source/CaptureDataConverter/Private/Nodes/ThirdPartyEncoder/CaptureConvertVideoDataThirdParty.cpp
+++ b/Engine/Plugins/VirtualProduction/CaptureManager/CaptureManagerApp/Source/CaptureDataConverter/Private/Nodes/ThirdPartyEncoder/CaptureConvertVideoDataThirdParty.cpp
@@ -453,7 +453,7 @@ FCaptureConvertVideoDataThirdParty::FResult FCaptureConvertVideoDataThirdParty::
 
 		if (CopyResult == COPY_Canceled)
 		{
-			FText Message = LOCTEXT("CaptureConvertVideoNodeTP_AbortedByUser", "Image data copy aborted by user");
+			FText Message = LOCTEXT("CaptureConvertVideoNodeTP_CopyAbortedByUser", "Image data copy aborted by user");
 			return MakeError(MoveTemp(Message));
 		}
 	}
