[Capture Manager] Check for empty string when cleaning up ingest directory.

#rnx

diff --git a/Engine/Plugins/VirtualProduction/CaptureManager/CaptureManagerApp/Source/IngestLiveLinkDevice/Private/BaseIngestLiveLinkDevice.cpp b/Engine/Plugins/VirtualProduction/CaptureManager/CaptureManagerApp/Source/IngestLiveLinkDevice/Private/BaseIngestLiveLinkDevice.cpp
--- a/Engine/Plugins/VirtualProduction/CaptureManager/CaptureManagerApp/Source/IngestLiveLinkDevice/Private/BaseIngestLiveLinkDevice.cpp
+++ b/Engine/Plugins/VirtualProduction/CaptureManager/CaptureManagerApp/Source/IngestLiveLinkDevice/Private/BaseIngestLiveLinkDevice.cpp
@@ -19,14 +19,47 @@ namespace UE::CaptureManager::Private
 {
 static void TryCleanUpWorkingDirectory(const FString& InWorkingDirectory)
 {
+	FString NormalizedFullPath = FPaths::ConvertRelativePathToFull(InWorkingDirectory);
+	FPaths::NormalizeFilename(NormalizedFullPath);
+
+	// Note: The raw input string is used for the empty check, not the normalized full path, as we are basically 
+	// trying to catch the situation when someone has forgotten to set it, before data loss occurs.
+
+	if (InWorkingDirectory.IsEmpty())
+	{
+		UE_LOG(
+			LogIngestLiveLinkDevice,
+			Error,
+			TEXT(
+				"Refused to clean up the ingest working directory. The supplied path was an empty string, which could "
+				"have resulted in the unintended removal of the current working directory: %s"
+			),
+			*NormalizedFullPath
+		);
+
+		return;
+	}
+
 	const UCaptureManagerSettings* CaptureManagerSettings = GetDefault<UCaptureManagerSettings>();
 	check(CaptureManagerSettings);
 
-	const bool bShouldCleanWorkingDirectory = CaptureManagerSettings->bShouldCleanWorkingDirectory;
-	if (bShouldCleanWorkingDirectory)
+	if (CaptureManagerSettings->bShouldCleanWorkingDirectory)
 	{
-		IFileManager::Get().DeleteDirectory(*InWorkingDirectory, false, true);
-		UE_LOG(LogIngestLiveLinkDevice, Display, TEXT("Deleting directory: %s"), *InWorkingDirectory);
+		UE_LOG(LogIngestLiveLinkDevice, Display, TEXT("Deleting directory: %s"), *NormalizedFullPath);
+
+		constexpr bool bRequireExists = false;
+		constexpr bool bIsRecursive = true;
+		const bool bIsDeleted = IFileManager::Get().DeleteDirectory(*NormalizedFullPath, bRequireExists, bIsRecursive);
+
+		if (bIsDeleted)
+		{
+			UE_LOG(LogIngestLiveLinkDevice, Display, TEXT("Directory deleted: %s"), *NormalizedFullPath);
+		}
+		else
+		{
+			const uint32 ErrorCode = FPlatformMisc::GetLastError();
+			UE_LOG(LogIngestLiveLinkDevice, Error, TEXT("Failed to delete directory (Error code: %u): %s"), ErrorCode, *NormalizedFullPath);
+		}
 	}
 }
 
