[Capture Manager] Use weak ptrs on endpoint manager interface.

This just makes the ownership more clear and reduces the chances of people holding onto endpoints for longer than they should.

#rnx

diff --git a/Engine/Plugins/VirtualProduction/CaptureManager/CaptureManagerApp/Source/CaptureManagerEditor/Private/IngestManagement/IngestJobProcessor.cpp b/Engine/Plugins/VirtualProduction/CaptureManager/CaptureManagerApp/Source/CaptureManagerEditor/Private/IngestManagement/IngestJobProcessor.cpp
--- a/Engine/Plugins/VirtualProduction/CaptureManager/CaptureManagerApp/Source/CaptureManagerEditor/Private/IngestManagement/IngestJobProcessor.cpp
+++ b/Engine/Plugins/VirtualProduction/CaptureManager/CaptureManagerApp/Source/CaptureManagerEditor/Private/IngestManagement/IngestJobProcessor.cpp
@@ -31,25 +31,30 @@ namespace Private
 		}
 
 		TSharedRef<FUnrealEndpointManager> UnrealEndpointManager = MaybeUnrealEndpointManager.GetValue();
-		const TArray<TSharedRef<FUnrealEndpoint>> Endpoints = UnrealEndpointManager->GetEndpoints();
+		const TArray<TWeakPtr<FUnrealEndpoint>> WeakEndpoints = UnrealEndpointManager->GetEndpoints();
 
-		for (const TSharedRef<FUnrealEndpoint>& Endpoint : Endpoints)
+		for (const TWeakPtr<FUnrealEndpoint>& WeakEndpoint : WeakEndpoints)
 		{
-			Endpoint->StopConnection();
+			if (const TSharedPtr<FUnrealEndpoint> Endpoint = WeakEndpoint.Pin())
+			{
+				Endpoint->StopConnection();
+			}
 		}
 
-		// Wait for all the connections to actually close
-		for (const TSharedRef<FUnrealEndpoint>& Endpoint : Endpoints)
+		// Give any extant endpoints a chance to close their connections
+		for (const TWeakPtr<FUnrealEndpoint>& WeakEndpoint : WeakEndpoints)
 		{
-			const bool bIsStopped = Endpoint->WaitForConnectionState(FUnrealEndpoint::EConnectionState::Disconnected, 3'000);
-
-			if (!bIsStopped)
+			if (const TSharedPtr<FUnrealEndpoint> Endpoint = WeakEndpoint.Pin())
 			{
-				FString EndpointInfoString = UnrealEndpointInfoToString(Endpoint->GetInfo());
-				UE_LOG(LogCaptureManagerIngestJobProcessor, Warning, TEXT("Failed to stop unreal endpoint: %s"), *EndpointInfoString);
+				const bool bIsStopped = Endpoint->WaitForConnectionState(FUnrealEndpoint::EConnectionState::Disconnected, 3'000);
+
+				if (!bIsStopped)
+				{
+					FString EndpointInfoString = UnrealEndpointInfoToString(Endpoint->GetInfo());
+					UE_LOG(LogCaptureManagerIngestJobProcessor, Warning, TEXT("Failed to stop unreal endpoint: %s"), *EndpointInfoString);
+				}
 			}
 		}
-
 	}
 } // namespace Private
 
diff --git a/Engine/Plugins/VirtualProduction/CaptureManager/CaptureManagerApp/Source/CaptureManagerEditor/Private/UI/SIngestSelectableUnrealEndpoint.cpp b/Engine/Plugins/VirtualProduction/CaptureManager/CaptureManagerApp/Source/CaptureManagerEditor/Private/UI/SIngestSelectableUnrealEndpoint.cpp
--- a/Engine/Plugins/VirtualProduction/CaptureManager/CaptureManagerApp/Source/CaptureManagerEditor/Private/UI/SIngestSelectableUnrealEndpoint.cpp
+++ b/Engine/Plugins/VirtualProduction/CaptureManager/CaptureManagerApp/Source/CaptureManagerEditor/Private/UI/SIngestSelectableUnrealEndpoint.cpp
@@ -199,15 +199,18 @@ TArray<TSharedRef<UE::CaptureManager::FUnrealEndpointInfo>> SIngestSelectableUnr
 	using namespace UE::CaptureManager;
 	using namespace UE::CaptureManager::Private;
 
-	TArray<TSharedRef<FUnrealEndpoint>> LatestEndpoints = UnrealEndpointManager->GetEndpoints();
+	TArray<TWeakPtr<FUnrealEndpoint>> WeakEndpoints = UnrealEndpointManager->GetEndpoints();
 
 	TArray<FUnrealEndpointInfo> SortedEndpointInfos;
-	SortedEndpointInfos.Reserve(LatestEndpoints.Num());
+	SortedEndpointInfos.Reserve(WeakEndpoints.Num());
 
-	for (const TSharedRef<FUnrealEndpoint>& Endpoint : LatestEndpoints)
+	for (const TWeakPtr<FUnrealEndpoint>& WeakEndpoint : WeakEndpoints)
 	{
-		FUnrealEndpointInfo EndpointInfo = Endpoint->GetInfo();
-		SortedEndpointInfos.Emplace(EndpointInfo);
+		if (const TSharedPtr<FUnrealEndpoint> Endpoint = WeakEndpoint.Pin())
+		{
+			FUnrealEndpointInfo EndpointInfo = Endpoint->GetInfo();
+			SortedEndpointInfos.Emplace(EndpointInfo);
+		}
 	}
 
 	// Keep the order deterministic
diff --git a/Engine/Plugins/VirtualProduction/CaptureManager/CaptureManagerApp/Source/CaptureManagerUnrealEndpoint/Private/CaptureManagerUnrealEndpointManager.cpp b/Engine/Plugins/VirtualProduction/CaptureManager/CaptureManagerApp/Source/CaptureManagerUnrealEndpoint/Private/CaptureManagerUnrealEndpointManager.cpp
--- a/Engine/Plugins/VirtualProduction/CaptureManager/CaptureManagerApp/Source/CaptureManagerUnrealEndpoint/Private/CaptureManagerUnrealEndpointManager.cpp
+++ b/Engine/Plugins/VirtualProduction/CaptureManager/CaptureManagerApp/Source/CaptureManagerUnrealEndpoint/Private/CaptureManagerUnrealEndpointManager.cpp
@@ -18,10 +18,10 @@ struct FUnrealEndpointManager::FImpl
 	void Start();
 	void Stop();
 
-	TSharedPtr<FUnrealEndpoint> WaitForEndpoint(TFunction<bool(const FUnrealEndpoint&)> InPredicate, int32 InTimeoutMS);
-	TSharedPtr<FUnrealEndpoint> FindEndpointByPredicate(TFunction<bool(const FUnrealEndpoint&)> InPredicate);
-	TArray<TSharedRef<FUnrealEndpoint>> GetEndpoints();
-	TArray<TSharedRef<FUnrealEndpoint>> FindEndpointsByPredicate(TFunction<bool(const FUnrealEndpoint&)> InPredicate);
+	TOptional<TWeakPtr<FUnrealEndpoint>> WaitForEndpoint(TFunction<bool(const FUnrealEndpoint&)> InPredicate, int32 InTimeoutMS);
+	TOptional<TWeakPtr<FUnrealEndpoint>> FindEndpointByPredicate(TFunction<bool(const FUnrealEndpoint&)> InPredicate);
+	TArray<TWeakPtr<FUnrealEndpoint>> GetEndpoints();
+	TArray<TWeakPtr<FUnrealEndpoint>> FindEndpointsByPredicate(TFunction<bool(const FUnrealEndpoint&)> InPredicate);
 	int32 GetNumEndpoints() const;
 
 	TArray<TSharedRef<FUnrealEndpoint>> Endpoints;
@@ -54,22 +54,22 @@ void FUnrealEndpointManager::Stop()
 	Impl->Stop();
 }
 
-TSharedPtr<FUnrealEndpoint> FUnrealEndpointManager::WaitForEndpoint(TFunction<bool(const FUnrealEndpoint&)> InPredicate, int32 InTimeoutMS)
+TOptional<TWeakPtr<FUnrealEndpoint>> FUnrealEndpointManager::WaitForEndpoint(TFunction<bool(const FUnrealEndpoint&)> InPredicate, int32 InTimeoutMS)
 {
 	return Impl->WaitForEndpoint(MoveTemp(InPredicate), InTimeoutMS);
 }
 
-TSharedPtr<FUnrealEndpoint> FUnrealEndpointManager::FindEndpointByPredicate(TFunction<bool(const FUnrealEndpoint&)> InPredicate)
+TOptional<TWeakPtr<FUnrealEndpoint>> FUnrealEndpointManager::FindEndpointByPredicate(TFunction<bool(const FUnrealEndpoint&)> InPredicate)
 {
 	return Impl->FindEndpointByPredicate(MoveTemp(InPredicate));
 }
 
-TArray<TSharedRef<FUnrealEndpoint>> FUnrealEndpointManager::FindEndpointsByPredicate(TFunction<bool(const FUnrealEndpoint&)> InPredicate)
+TArray<TWeakPtr<FUnrealEndpoint>> FUnrealEndpointManager::FindEndpointsByPredicate(TFunction<bool(const FUnrealEndpoint&)> InPredicate)
 {
 	return Impl->FindEndpointsByPredicate(MoveTemp(InPredicate));
 }
 
-TArray<TSharedRef<FUnrealEndpoint>> FUnrealEndpointManager::GetEndpoints()
+TArray<TWeakPtr<FUnrealEndpoint>> FUnrealEndpointManager::GetEndpoints()
 {
 	return Impl->GetEndpoints();
 }
@@ -181,9 +181,9 @@ void FUnrealEndpointManager::FImpl::Stop()
 	bIsRunning = false;
 }
 
-TSharedPtr<FUnrealEndpoint> FUnrealEndpointManager::FImpl::WaitForEndpoint(TFunction<bool(const FUnrealEndpoint&)> InPredicate, const int32 InTimeoutMS)
+TOptional<TWeakPtr<FUnrealEndpoint>> FUnrealEndpointManager::FImpl::WaitForEndpoint(TFunction<bool(const FUnrealEndpoint&)> InPredicate, const int32 InTimeoutMS)
 {
-	TSharedPtr<FUnrealEndpoint> Endpoint;
+	TOptional<TWeakPtr<FUnrealEndpoint>> Endpoint;
 
 	std::unique_lock<std::mutex> Lock(Mutex);
 
@@ -202,52 +202,63 @@ TSharedPtr<FUnrealEndpoint> FUnrealEndpointManager::FImpl::WaitForEndpoint(TFunc
 
 			if (FoundEndpoint)
 			{
-				Endpoint = *FoundEndpoint;
+				Endpoint = FoundEndpoint->ToWeakPtr();
 			}
 
-			return Endpoint != nullptr;
+			return Endpoint.IsSet();
 		}
 	);
 
 	return Endpoint;
 }
 
-TSharedPtr<FUnrealEndpoint> FUnrealEndpointManager::FImpl::FindEndpointByPredicate(TFunction<bool(const FUnrealEndpoint&)> InPredicate)
+TOptional<TWeakPtr<FUnrealEndpoint>> FUnrealEndpointManager::FImpl::FindEndpointByPredicate(TFunction<bool(const FUnrealEndpoint&)> InPredicate)
 {
+	TOptional<TWeakPtr<FUnrealEndpoint>> Endpoint;
+
 	std::lock_guard<std::mutex> LockGuard(Mutex);
 
-	TSharedRef<FUnrealEndpoint>* Endpoint = Endpoints.FindByPredicate(
+	TSharedRef<FUnrealEndpoint>* FoundEndpoint = Endpoints.FindByPredicate(
 		[&InPredicate](const TSharedRef<FUnrealEndpoint>& InEndpoint)
 		{
 			return InPredicate(*InEndpoint);
 		}
 	);
 
-	if (Endpoint)
+	if (FoundEndpoint)
 	{
-		return (*Endpoint).ToSharedPtr();
+		Endpoint = FoundEndpoint->ToWeakPtr();
 	}
 
-	return nullptr;
+	return Endpoint;
 }
 
-TArray<TSharedRef<FUnrealEndpoint>> FUnrealEndpointManager::FImpl::GetEndpoints()
+TArray<TWeakPtr<FUnrealEndpoint>> FUnrealEndpointManager::FImpl::GetEndpoints()
 {
+	TArray<TWeakPtr<FUnrealEndpoint>> WeakEndpoints;
+
 	std::lock_guard<std::mutex> LockGuard(Mutex);
-	return Endpoints;
+	WeakEndpoints.Reserve(Endpoints.Num());
+
+	for (const TSharedRef<FUnrealEndpoint>& Endpoint : Endpoints)
+	{
+		WeakEndpoints.Emplace(Endpoint.ToWeakPtr());
+	}
+
+	return WeakEndpoints;
 }
 
-TArray<TSharedRef<FUnrealEndpoint>> FUnrealEndpointManager::FImpl::FindEndpointsByPredicate(TFunction<bool(const FUnrealEndpoint&)> InPredicate)
+TArray<TWeakPtr<FUnrealEndpoint>> FUnrealEndpointManager::FImpl::FindEndpointsByPredicate(TFunction<bool(const FUnrealEndpoint&)> InPredicate)
 {
-	std::lock_guard<std::mutex> LockGuard(Mutex);
+	TArray<TWeakPtr<FUnrealEndpoint>> FilteredEndpoints;
 
-	TArray<TSharedRef<FUnrealEndpoint>> FilteredEndpoints;
+	std::lock_guard<std::mutex> LockGuard(Mutex);
 
 	for (const TSharedRef<FUnrealEndpoint>& Endpoint : Endpoints)
 	{
 		if (InPredicate(*Endpoint))
 		{
-			FilteredEndpoints.Emplace(Endpoint);
+			FilteredEndpoints.Emplace(Endpoint.ToWeakPtr());
 		}
 	}
 
diff --git a/Engine/Plugins/VirtualProduction/CaptureManager/CaptureManagerApp/Source/CaptureManagerUnrealEndpoint/Private/UCaptureManagerUnrealEndpointManager.cpp b/Engine/Plugins/VirtualProduction/CaptureManager/CaptureManagerApp/Source/CaptureManagerUnrealEndpoint/Private/UCaptureManagerUnrealEndpointManager.cpp
--- a/Engine/Plugins/VirtualProduction/CaptureManager/CaptureManagerApp/Source/CaptureManagerUnrealEndpoint/Private/UCaptureManagerUnrealEndpointManager.cpp
+++ b/Engine/Plugins/VirtualProduction/CaptureManager/CaptureManagerApp/Source/CaptureManagerUnrealEndpoint/Private/UCaptureManagerUnrealEndpointManager.cpp
@@ -35,7 +35,7 @@ bool UCaptureManagerUnrealEndpointManager::WaitForEndpointByHostName(const FStri
 {
 	using namespace UE::CaptureManager;
 
-	TSharedPtr<FUnrealEndpoint> NativeEndpoint = Impl->EndpointManager->WaitForEndpoint(
+	TOptional<TWeakPtr<FUnrealEndpoint>> NativeEndpoint = Impl->EndpointManager->WaitForEndpoint(
 		[&InHostName](const FUnrealEndpoint& InEndpoint)
 		{
 			return InEndpoint.GetInfo().HostName == InHostName;
@@ -43,7 +43,7 @@ bool UCaptureManagerUnrealEndpointManager::WaitForEndpointByHostName(const FStri
 		InTimeoutMS
 	);
 
-	return NativeEndpoint != nullptr;
+	return NativeEndpoint.IsSet();
 }
 
 UCaptureManagerUnrealEndpointManager::FImpl::FImpl() :
diff --git a/Engine/Plugins/VirtualProduction/CaptureManager/CaptureManagerApp/Source/CaptureManagerUnrealEndpoint/Public/CaptureManagerUnrealEndpointManager.h b/Engine/Plugins/VirtualProduction/CaptureManager/CaptureManagerApp/Source/CaptureManagerUnrealEndpoint/Public/CaptureManagerUnrealEndpointManager.h
--- a/Engine/Plugins/VirtualProduction/CaptureManager/CaptureManagerApp/Source/CaptureManagerUnrealEndpoint/Public/CaptureManagerUnrealEndpointManager.h
+++ b/Engine/Plugins/VirtualProduction/CaptureManager/CaptureManagerApp/Source/CaptureManagerUnrealEndpoint/Public/CaptureManagerUnrealEndpointManager.h
@@ -35,26 +35,26 @@ public:
 	/**
 	* @brief Blocks until an endpoint matching the given predicate is discovered or the timeout (in milliseconds) is reached.
 	*
-	* @returns The endpoint (if discovered) else nullptr.
+	* @returns An optional containing a weak pointer to the endpoint if it was discovered else the optional will be empty.
 	*/
-	TSharedPtr<FUnrealEndpoint> WaitForEndpoint(TFunction<bool(const FUnrealEndpoint&)> InPredicate, int32 InTimeoutMS);
+	TOptional<TWeakPtr<FUnrealEndpoint>> WaitForEndpoint(TFunction<bool(const FUnrealEndpoint&)> InPredicate, int32 InTimeoutMS);
 
 	/**
 	* @brief Finds a discovered endpoint matching the given predicate.
 	*
-	* @returns The endpoint (if found) else nullptr.
+	* @returns An optional containing a weak pointer to the endpoint if it is found else the optional will be empty.
 	*/
-	TSharedPtr<FUnrealEndpoint> FindEndpointByPredicate(TFunction<bool(const FUnrealEndpoint&)> InPredicate);
+	TOptional<TWeakPtr<FUnrealEndpoint>> FindEndpointByPredicate(TFunction<bool(const FUnrealEndpoint&)> InPredicate);
 
 	/**
-	* @returns A list of all discovered endpoints matching the given predicate.
+	* @returns A list of weak pointers to the discovered endpoints which match the predicate.
 	*/
-	TArray<TSharedRef<FUnrealEndpoint>> FindEndpointsByPredicate(TFunction<bool(const FUnrealEndpoint&)> InPredicate);
+	TArray<TWeakPtr<FUnrealEndpoint>> FindEndpointsByPredicate(TFunction<bool(const FUnrealEndpoint&)> InPredicate);
 
 	/**
-	* @returns A list of all discovered endpoints.
+	* @returns A list of weak pointers to the discovered endpoints.
 	*/
-	TArray<TSharedRef<FUnrealEndpoint>> GetEndpoints();
+	TArray<TWeakPtr<FUnrealEndpoint>> GetEndpoints();
 
 	/**
 	* @returns The number of discovered endpoints
diff --git a/Engine/Plugins/VirtualProduction/CaptureManager/CaptureManagerApp/Source/IngestLiveLinkDevice/Private/BaseIngestLiveLinkDevice.cpp b/Engine/Plugins/VirtualProduction/CaptureManager/CaptureManagerApp/Source/IngestLiveLinkDevice/Private/BaseIngestLiveLinkDevice.cpp
--- a/Engine/Plugins/VirtualProduction/CaptureManager/CaptureManagerApp/Source/IngestLiveLinkDevice/Private/BaseIngestLiveLinkDevice.cpp
+++ b/Engine/Plugins/VirtualProduction/CaptureManager/CaptureManagerApp/Source/IngestLiveLinkDevice/Private/BaseIngestLiveLinkDevice.cpp
@@ -203,15 +203,15 @@ void UBaseIngestLiveLinkDevice::FImpl::IngestTake(const UIngestCapability_Proces
 
 		TSharedRef<FUnrealEndpointManager> UnrealEndpointManager = FModuleManager::LoadModuleChecked<FCaptureManagerUnrealEndpointModule>("CaptureManagerUnrealEndpoint").GetEndpointManager();
 
-		TSharedPtr<FUnrealEndpoint> UnrealEndpoint = UnrealEndpointManager->FindEndpointByPredicate(
+		TOptional<TWeakPtr<FUnrealEndpoint>> WeakUnrealEndpoint = UnrealEndpointManager->FindEndpointByPredicate(
 			[&IngestOptions](const FUnrealEndpoint& InEndpoint)
 			{
 				return InEndpoint.GetInfo().HostName == IngestOptions->UploadHostName;
 			}
 		);
 
-		// We can't upload without an unreal endpoint, so there's no point in proceeding any further
-		if (!UnrealEndpoint)
+		// The endpoint was not found at all
+		if (!WeakUnrealEndpoint)
 		{
 			TryCleanUpWorkingDirectory(Params.TakeOutputDirectory);
 
@@ -220,6 +220,18 @@ void UBaseIngestLiveLinkDevice::FImpl::IngestTake(const UIngestCapability_Proces
 			return;
 		}
 
+		TSharedPtr<FUnrealEndpoint> UnrealEndpoint = WeakUnrealEndpoint->Pin();
+
+		// The endpoint was found but has since expired
+		if (!UnrealEndpoint)
+		{
+			TryCleanUpWorkingDirectory(Params.TakeOutputDirectory);
+
+			FString Message = FString::Printf(TEXT("Upload failed because the requested endpoint has expired: %s"), *IngestOptions->UploadHostName);
+			Owner->ExecuteProcessFinishedReporter(InProcessHandle, MakeError(FIngestCapability_Error::UnrealEndpointNotFound, MoveTemp(Message)));
+			return;
+		}
+
 		// We wait until after the download and conversion to start the connection, in this way we won't require a
 		// connection for download only workflows. Note: Disconnection occurs for all endpoints at the end of job
 		// processing (in the ingest job processor).
@@ -324,21 +336,23 @@ void UBaseIngestLiveLinkDevice::FImpl::CancelIngest(int32 InTakeId)
 		// Find the endpoint responsible for uploading this task
 		TSharedRef<FUnrealEndpointManager> UnrealEndpointManager = FModuleManager::LoadModuleChecked<FCaptureManagerUnrealEndpointModule>("CaptureManagerUnrealEndpoint").GetEndpointManager();
 
-		TSharedPtr<FUnrealEndpoint> Endpoint = UnrealEndpointManager->FindEndpointByPredicate(
+		TOptional<TWeakPtr<FUnrealEndpoint>> WeakUnrealEndpoint = UnrealEndpointManager->FindEndpointByPredicate(
 			[TakeUploadInfo](const FUnrealEndpoint& InEndpoint) -> bool
 			{
 				return InEndpoint.GetInfo().EndpointID == TakeUploadInfo->EndpointID;
 			}
 		);
 
-		if (Endpoint)
+		if (WeakUnrealEndpoint)
 		{
-			Endpoint->CancelTakeUploadTask(TakeUploadInfo->TaskId);
+			if (const TSharedPtr<FUnrealEndpoint> UnrealEndpoint = WeakUnrealEndpoint->Pin())
+			{
+				UnrealEndpoint->CancelTakeUploadTask(TakeUploadInfo->TaskId);
+			}
 		}
 	}
 }
 
-
 FCaptureDataConverterResult<void> UBaseIngestLiveLinkDevice::FImpl::RunConversion(int32 InTakeId,
 																				  TStrongObjectPtr<const UIngestCapability_Options> InIngestOptions,
 																				  TSharedPtr<UE::CaptureManager::FTaskProgress> InTaskProgress,
