Setting video frame rate from the Electra Player

diff --git a/Engine/Plugins/VirtualProduction/CaptureManager/CaptureManagerApp/Source/IngestLiveLinkDevice/IngestLiveLinkDevice.Build.cs b/Engine/Plugins/VirtualProduction/CaptureManager/CaptureManagerApp/Source/IngestLiveLinkDevice/IngestLiveLinkDevice.Build.cs
--- a/Engine/Plugins/VirtualProduction/CaptureManager/CaptureManagerApp/Source/IngestLiveLinkDevice/IngestLiveLinkDevice.Build.cs
+++ b/Engine/Plugins/VirtualProduction/CaptureManager/CaptureManagerApp/Source/IngestLiveLinkDevice/IngestLiveLinkDevice.Build.cs
@@ -20,6 +20,7 @@ public class IngestLiveLinkDevice : ModuleRules
 		PrivateDependencyModuleNames.AddRange(new string[]
 		{
 			"CaptureProtocolStack",
+			"CaptureDataUtils",
 			"CaptureDataConverter",
 			"CaptureManagerUnrealEndpoint",
 			"Engine",
diff --git a/Engine/Plugins/VirtualProduction/CaptureManager/CaptureManagerApp/Source/IngestLiveLinkDevice/Private/Utils/CaptureExtractTimecode.cpp b/Engine/Plugins/VirtualProduction/CaptureManager/CaptureManagerApp/Source/IngestLiveLinkDevice/Private/Utils/CaptureExtractTimecode.cpp
--- a/Engine/Plugins/VirtualProduction/CaptureManager/CaptureManagerApp/Source/IngestLiveLinkDevice/Private/Utils/CaptureExtractTimecode.cpp
+++ b/Engine/Plugins/VirtualProduction/CaptureManager/CaptureManagerApp/Source/IngestLiveLinkDevice/Private/Utils/CaptureExtractTimecode.cpp
@@ -14,29 +14,63 @@
 #include "IMediaPlayer.h"
 #include "IMediaOptions.h"
 #include "IMediaEventSink.h"
+#include "IMediaTracks.h"
+
+#include "ParseTakeUtils.h"
 
 #define LOCTEXT_NAMESPACE "CaptureExtractTimecode"
 
 DEFINE_LOG_CATEGORY_STATIC(LogCaptureExtractTimecode, Log, All);
 
-FCaptureExtractVideoTimecode::FCaptureExtractVideoTimecode(const FString& InFilePath)
-	: FilePath(InFilePath)
+FCaptureExtractVideoInfo::FResult FCaptureExtractVideoInfo::Create(const FString& InFilePath)
 {
-	// Replace any backslashes to avoid escape code issues while logging the file path
-	FilePath = FilePath.Replace(TEXT("\\"), TEXT("/"));
+	using namespace UE::CaptureManager;
+
+	FCaptureExtractVideoInfo VideoInfoExtractor;
+
+	FExtractResult Result = VideoInfoExtractor.ExtractInfo(InFilePath);
+	if (Result.HasError())
+	{
+		return MakeError(Result.StealError());
+	}
+
+	return MakeValue(MoveTemp(VideoInfoExtractor));
 }
 
-UE::CaptureManager::FTimecodeAndRateResult FCaptureExtractVideoTimecode::Extract()
+FFrameRate FCaptureExtractVideoInfo::GetFrameRate() const
 {
-	using namespace UE::CaptureManager;
+	return FrameRate;
+}
 
-	check(!FilePath.IsEmpty())
+FTimecode FCaptureExtractVideoInfo::GetTimecode() const
+{
+	return TimecodeInfo.Timecode;
+}
+
+FFrameRate FCaptureExtractVideoInfo::GetTimecodeRate() const
+{
+	return TimecodeInfo.TimecodeRate;
+}
+
+FCaptureExtractVideoInfo::FCaptureExtractVideoInfo()
+	: FrameRate(1, 1)
+{
+}
+
+FCaptureExtractVideoInfo::FExtractResult FCaptureExtractVideoInfo::ExtractInfo(const FString& InFilePath)
+{
+	using namespace UE::CaptureManager;
 
-	FTimecodeAndRateResult Result = MakeError(ECaptureExtractTimecodeError::UnhandledMedia);
+	// Replace any backslashes to avoid escape code issues while logging the file path
+	const FString FilePath = InFilePath.Replace(TEXT("\\"), TEXT("/"));
+	if (FilePath.IsEmpty())
+	{
+		return MakeError(ECaptureExtractInfoError::InternalError);
+	}
 
-	Result = ExtractTimecodeUsingElectraPlayer();
+	FCaptureExtractVideoInfo::FExtractResult Result = ExtractInfoUsingElectraPlayer(InFilePath);;
 
-	if (!Result.IsValid())
+	if (Result.HasError())
 	{
 		const UCaptureManagerSettings* Settings = GetDefault<UCaptureManagerSettings>();
 
@@ -46,14 +80,14 @@ UE::CaptureManager::FTimecodeAndRateResult FCaptureExtractVideoTimecode::Extract
 			if (EncoderPath.EndsWith(TEXT("ffmpeg.exe")))
 			{
 				EncoderPath.ReplaceInline(TEXT("ffmpeg.exe"), TEXT("ffprobe.exe"));
-				Result = ExtractTimecodeUsingFFProbe(EncoderPath);
+				Result = ExtractInfoUsingFFProbe(InFilePath, EncoderPath);
 			}
 		}
 	}
 
 	if (Result.HasValue())
 	{
-		UE_LOG(LogCaptureExtractTimecode, Display, TEXT("Extracted timecode '%s' from video file: %s"), *Result.GetValue().Timecode.ToString(), *FilePath);
+		UE_LOG(LogCaptureExtractTimecode, Display, TEXT("Extracted timecode '%s' from video file: %s"), *TimecodeInfo.Timecode.ToString(), *FilePath);
 	}
 	else
 	{
@@ -63,12 +97,12 @@ UE::CaptureManager::FTimecodeAndRateResult FCaptureExtractVideoTimecode::Extract
 	return Result;
 }
 
-UE::CaptureManager::FTimecodeAndRateResult FCaptureExtractVideoTimecode::ExtractTimecodeUsingElectraPlayer()
+FCaptureExtractVideoInfo::FExtractResult FCaptureExtractVideoInfo::ExtractInfoUsingElectraPlayer(const FString& InFilePath)
 {
 	using namespace UE::CaptureManager;
 
-	using FEventSinkResult = TValueOrError<void, ECaptureExtractTimecodeError>;
-	TPromise<FEventSinkResult> EventSinkPromise;
+	using FEventSinkResult = FExtractResult;
+	TPromise<FExtractResult> EventSinkPromise;
 
 	class FMediaEventSink : public IMediaEventSink
 	{
@@ -93,7 +127,7 @@ UE::CaptureManager::FTimecodeAndRateResult FCaptureExtractVideoTimecode::Extract
 			}
 			else if (Event == EMediaEvent::MediaOpenFailed)
 			{
-				Promise.SetValue(MakeError(ECaptureExtractTimecodeError::UnableToOpenMedia));
+				Promise.SetValue(MakeError(ECaptureExtractInfoError::UnableToOpenMedia));
 			}
 		}
 
@@ -115,15 +149,15 @@ UE::CaptureManager::FTimecodeAndRateResult FCaptureExtractVideoTimecode::Extract
 	if (!MediaPlayer)
 	{
 		// Broken promises are considered errors
-		UE::CaptureManager::FTimecodeAndRateResult Result = MakeError(ECaptureExtractTimecodeError::InternalError);
-		EventSinkPromise.SetValue(MakeError(Result.GetError()));
-		return Result;
+		FExtractResult Error = MakeError(ECaptureExtractInfoError::InternalError);
+		EventSinkPromise.SetValue(Error);
+		return Error;
 	}
 
 	UDesiredPlayerMediaSource* MediaSource = NewObject<UDesiredPlayerMediaSource>();
-	MediaSource->SetFilePath(FilePath);
+	MediaSource->SetFilePath(InFilePath);
 
-	FString FileUrl = TEXT("file://") + FilePath;
+	FString FileUrl = TEXT("file://") + InFilePath;
 
 	FMediaPlayerOptions PlayerOptions;
 	PlayerOptions.SetAllAsOptional();
@@ -137,7 +171,9 @@ UE::CaptureManager::FTimecodeAndRateResult FCaptureExtractVideoTimecode::Extract
 	FDateTime EventWaitStart = FDateTime::Now();
 
 	TFuture<FEventSinkResult> EventSinkFuture = EventSinkPromise.GetFuture();
-	
+
+	bool bIsTimeout = true;
+
 	while ((FDateTime::Now() - EventWaitStart).GetSeconds() < TimeoutPeriod)
 	{
 		MediaPlayer->TickInput(0, 0);
@@ -147,9 +183,11 @@ UE::CaptureManager::FTimecodeAndRateResult FCaptureExtractVideoTimecode::Extract
 			FEventSinkResult Result = EventSinkFuture.Get();
 			if (Result.HasError())
 			{
-				return MakeError(Result.StealError());
+				return Result;
 			}
 
+			bIsTimeout = false;
+
 			const FVariant Timecode = MediaPlayer->GetMediaInfo(UMediaPlayer::MediaInfoNameStartTimecodeValue.Resolve());
 			if (!Timecode.IsEmpty())
 			{
@@ -162,6 +200,21 @@ UE::CaptureManager::FTimecodeAndRateResult FCaptureExtractVideoTimecode::Extract
 				}
 			}
 
+			int32 Track = MediaPlayer->GetTracks().GetSelectedTrack(EMediaTrackType::Video);
+			int32 Format = MediaPlayer->GetTracks().GetTrackFormat(EMediaTrackType::Video, Track);
+
+			if (Track != INDEX_NONE && Format != INDEX_NONE)
+			{
+				FMediaVideoTrackFormat FormatInfo;
+				MediaPlayer->GetTracks().GetVideoTrackFormat(Track, Format, FormatInfo);
+
+				FrameRate = ConvertFrameRate(FormatInfo.FrameRate);
+			}
+			else
+			{
+				UE_LOG(LogCaptureExtractTimecode, Warning, TEXT("Failed to obtain the frame rate"));
+			}
+
 			break;
 		}
 	}
@@ -170,100 +223,58 @@ UE::CaptureManager::FTimecodeAndRateResult FCaptureExtractVideoTimecode::Extract
 	MediaEventSink.Release();
 	MediaPlayer->Close();
 
-	if (TimecodeString.IsEmpty())
+	if (bIsTimeout)
 	{
-		UE::CaptureManager::FTimecodeAndRateResult Result = MakeError(ECaptureExtractTimecodeError::TimecodeNotFound);
+		FExtractResult Error = MakeError(ECaptureExtractInfoError::InternalError);
 
 		if (!EventSinkFuture.IsReady())
 		{
 			// Broken promises are considered errors
-			EventSinkPromise.SetValue(MakeError(Result.GetError()));
+			EventSinkPromise.SetValue(Error);
 		}
 
-		return Result;
+		return Error;
 	}
-	
-	const TOptional<FTimecode> Timecode = FTimecode::ParseTimecode(*TimecodeString);
-	if (!Timecode.IsSet())
+
+	if (!TimecodeString.IsEmpty())
 	{
-		return MakeError(ECaptureExtractTimecodeError::UnableToParseTimecode);
+		const TOptional<FTimecode> Timecode = FTimecode::ParseTimecode(*TimecodeString);
+		if (!Timecode.IsSet())
+		{
+			UE_LOG(LogCaptureExtractTimecode, Warning, TEXT("Failed to parse the timecode"));
+		}
+		else
+		{
+			TimecodeInfo.Timecode = Timecode.GetValue();
+		}
+	}
+	else
+	{
+		UE_LOG(LogCaptureExtractTimecode, Warning, TEXT("Timecode has not been found"));
 	}
-
-	TOptional<FFrameRate> MaybeTimecodeRate;
 
 	if (!TimecodeRateString.IsEmpty())
 	{
-		TValueOrError<FFrameRate, UE::CaptureManager::ECaptureExtractTimecodeError> Result = ParseTimecodeRate(TimecodeRateString);
+		TValueOrError<FFrameRate, UE::CaptureManager::ECaptureExtractInfoError> Result = ParseTimecodeRate(TimecodeRateString);
 		if (Result.IsValid())
 		{
-			MaybeTimecodeRate = Result.GetValue();
+			TimecodeInfo.TimecodeRate = Result.GetValue();
 		}
 	}
-
-	return MakeValue(Timecode.GetValue(), MoveTemp(MaybeTimecodeRate));
-}
-
-TValueOrError<FFrameRate, UE::CaptureManager::ECaptureExtractTimecodeError> FCaptureExtractVideoTimecode::ParseTimecodeRate(FString TimecodeRateString)
-{
-	using namespace UE::CaptureManager;
-
-	if (TimecodeRateString.IsEmpty())
-	{
-		return MakeError(ECaptureExtractTimecodeError::TimecodeRateNotFound);
-	}
 	else
 	{
-		// TimecodeRateString is made out of FFrameRate::ToPrettyText()
-		// It could be "{0} fps" or "{0} s"
-
-		FFrameRate TimecodeRate;
-
-		FString Left;
-		FString Right;
-		TimecodeRateString.Split(" ", &Left, &Right);
-
-		double Number = FCString::Atoi(*Left);
-		uint32 IntNumber;
-		uint32 Multiplier = 1;
-
-
-		for (; Multiplier <= 10000; Multiplier *= 10)
-		{
-			double TmpNumber = Number * Multiplier;
-			IntNumber = FMath::RoundToInt32(TmpNumber);
-			if (FMath::Abs(TmpNumber - IntNumber) < 0.01)
-			{
-				break;
-			}
-		}
-
-		uint32 Nominator = 0;
-		uint32 Denominator = 0;
-
-		if (Right == "fps")
-		{
-			Nominator = IntNumber;
-			Denominator = Multiplier;
-		}
-		else if (Right == "s")
-		{
-			Nominator = Multiplier;
-			Denominator = IntNumber;
-		}
-
-		TimecodeRate = FFrameRate(Nominator, Denominator);
-
-		return MakeValue(TimecodeRate);
+		UE_LOG(LogCaptureExtractTimecode, Warning, TEXT("Timecode rate has not been found. Using video Frame rate as timecode rate"));
+		TimecodeInfo.TimecodeRate = FrameRate;
 	}
+
+	return MakeValue();
 }
 
-UE::CaptureManager::FTimecodeAndRateResult FCaptureExtractVideoTimecode::ExtractTimecodeUsingFFProbe(const FString& InFFProbePath)
+FCaptureExtractVideoInfo::FExtractResult FCaptureExtractVideoInfo::ExtractInfoUsingFFProbe(const FString& InFilePath, const FString& InFFProbePath)
 {
 	using namespace UE::CaptureManager;
 
-	FTimecodeAndRateResult Result = MakeError(ECaptureExtractTimecodeError::UnhandledMedia);
-
-	FString CommandArgs = TEXT("-v error -select_streams v:0 -show_entries stream_tags=timecode:stream=r_frame_rate -of default=noprint_wrappers=1:nokey=1 ") + FilePath;
+	FString CommandArgs = TEXT("-v error -select_streams v:0 -show_entries stream_tags=timecode:stream=r_frame_rate -of default=noprint_wrappers=1:nokey=1 ") + InFilePath;
 
 	void* ReadPipe = nullptr;
 	void* WritePipe = nullptr;
@@ -285,7 +296,7 @@ UE::CaptureManager::FTimecodeAndRateResult FCaptureExtractVideoTimecode::Extract
 	if (!ProcHandle.IsValid())
 	{
 		UE_LOG(LogCaptureExtractTimecode, Display, TEXT("Extract timecode: Failed to start the process %s %s"), *InFFProbePath, *CommandArgs);
-		return MakeError(ECaptureExtractTimecodeError::InternalError);
+		return MakeError(ECaptureExtractInfoError::InternalError);
 	}
 
 	const FDateTime WaitStart = FDateTime::Now();
@@ -321,7 +332,6 @@ UE::CaptureManager::FTimecodeAndRateResult FCaptureExtractVideoTimecode::Extract
 	int32 ReturnCode = 0;
 	FPlatformProcess::GetProcReturnCode(ProcHandle, &ReturnCode);
 
-
 	TArray<uint8> CommandOutput;
 	{
 		bool Read = FPlatformProcess::ReadPipeToArray(ReadPipe, CommandOutput);
@@ -332,90 +342,132 @@ UE::CaptureManager::FTimecodeAndRateResult FCaptureExtractVideoTimecode::Extract
 	}
 	FullCommandOutput.Append(CommandOutput);
 
-	if (ReturnCode == 0)
+	if (ReturnCode != 0)
 	{
-		FTimecode Timecode;
-		bool bFoundTimecode = false;
-		FFrameRate FrameRate;
-		bool bFoundFrameRate = false;
+		if (!FullCommandOutput.IsEmpty())
+		{
+			UE_LOG(LogCaptureExtractTimecode, Display, TEXT("Failed to run the command: %s %s"), *InFFProbePath, *CommandArgs);
 
-		FString CommandOutputStr = FString::ConstructFromPtrSize(reinterpret_cast<const UTF8CHAR*>(FullCommandOutput.GetData()), FullCommandOutput.Num());
+			FString CommandOutputStr = FString::ConstructFromPtrSize(reinterpret_cast<const UTF8CHAR*>(FullCommandOutput.GetData()), FullCommandOutput.Num());
+			UE_LOG(LogCaptureExtractTimecode, Display,
+				   TEXT("Output from the command:\n>>>>>>\n%s<<<<<<"), *CommandOutputStr);
+		}
 
-		// Split the timecode and frame rate lines from the output
-		TArray<FString> OutputLinesStr;
-		CommandOutputStr.ParseIntoArrayLines(OutputLinesStr);
+		return MakeError(ECaptureExtractInfoError::InternalError);
+	}
+	bool bFoundTimecode = false;
+	bool bFoundFrameRate = false;
 
-		for (const FString& Line : OutputLinesStr)
-		{
-			if (!bFoundFrameRate)
-			{
-				TArray<FString> FrameRatePartsStr;
-				Line.ParseIntoArray(FrameRatePartsStr, TEXT("/"));
-				if (FrameRatePartsStr.Num() == 2)
-				{
-					int32 Numerator = FCString::Atoi(*FrameRatePartsStr[0]);
-					int32 Denominator = FCString::Atoi(*FrameRatePartsStr[1]);
+	FString CommandOutputStr = FString::ConstructFromPtrSize(reinterpret_cast<const UTF8CHAR*>(FullCommandOutput.GetData()), FullCommandOutput.Num());
 
-					FrameRate = FFrameRate(Numerator, Denominator);
-					bFoundFrameRate = true;
-					continue;
-				}
-			}
+	// Split the timecode and frame rate lines from the output
+	TArray<FString> OutputLinesStr;
+	CommandOutputStr.ParseIntoArrayLines(OutputLinesStr);
 
-			if (!bFoundTimecode)
+	for (const FString& Line : OutputLinesStr)
+	{
+		if (!bFoundFrameRate)
+		{
+			TArray<FString> FrameRatePartsStr;
+			Line.ParseIntoArray(FrameRatePartsStr, TEXT("/"));
+			if (FrameRatePartsStr.Num() == 2)
 			{
-				TArray<FString> TimecodePartsStr;
-				Line.ParseIntoArray(TimecodePartsStr, TEXT(":")); // We don't support DF timecode
-				if (TimecodePartsStr.Num() == 4)
-				{
-					int32 Hours = FCString::Atoi(*TimecodePartsStr[0]);
-					int32 Minutes = FCString::Atoi(*TimecodePartsStr[1]);
-					int32 Seconds = FCString::Atoi(*TimecodePartsStr[2]);
-					int32 Frames = FCString::Atoi(*TimecodePartsStr[3]);
-
-					constexpr bool bIsDropFrame = false;
-					Timecode = FTimecode(Hours, Minutes, Seconds, Frames, bIsDropFrame);
+				int32 Numerator = FCString::Atoi(*FrameRatePartsStr[0]);
+				int32 Denominator = FCString::Atoi(*FrameRatePartsStr[1]);
 
-					bFoundTimecode = true;
-					continue;
-				}
+				FrameRate = FFrameRate(Numerator, Denominator);
+				bFoundFrameRate = true;
+				continue;
 			}
 		}
 
 		if (!bFoundTimecode)
 		{
-			return MakeError(ECaptureExtractTimecodeError::UnableToParseTimecode);
+			TArray<FString> TimecodePartsStr;
+			Line.ParseIntoArray(TimecodePartsStr, TEXT(":")); // We don't support DF timecode
+			if (TimecodePartsStr.Num() == 4)
+			{
+				int32 Hours = FCString::Atoi(*TimecodePartsStr[0]);
+				int32 Minutes = FCString::Atoi(*TimecodePartsStr[1]);
+				int32 Seconds = FCString::Atoi(*TimecodePartsStr[2]);
+				int32 Frames = FCString::Atoi(*TimecodePartsStr[3]);
+
+				constexpr bool bIsDropFrame = false;
+				TimecodeInfo.Timecode = FTimecode(Hours, Minutes, Seconds, Frames, bIsDropFrame);
+				bFoundTimecode = true;
+				continue;
+			}
 		}
+	}
 
-		if (!bFoundFrameRate)
-		{
-			return MakeError(ECaptureExtractTimecodeError::UnableToParseTimecodeRate);
-		}
+	if (bFoundFrameRate)
+	{
+		UE_LOG(LogCaptureExtractTimecode, Warning, TEXT("Timecode rate has not been found. Using video frame rate as timecode rate"));
+		TimecodeInfo.TimecodeRate = FrameRate;
+	}
 
-		Result = MakeValue(FTimecodeAndRate{ Timecode , FrameRate });
+	return MakeValue();
+}
+
+TValueOrError<FFrameRate, UE::CaptureManager::ECaptureExtractInfoError> FCaptureExtractVideoInfo::ParseTimecodeRate(FString TimecodeRateString)
+{
+	using namespace UE::CaptureManager;
+
+	if (TimecodeRateString.IsEmpty())
+	{
+		return MakeError(ECaptureExtractInfoError::TimecodeRateNotFound);
 	}
 	else
 	{
-		if (!FullCommandOutput.IsEmpty())
+		// TimecodeRateString is made out of FFrameRate::ToPrettyText()
+		// It could be "{0} fps" or "{0} s"
+
+		FFrameRate TimecodeRate;
+
+		FString Left;
+		FString Right;
+		TimecodeRateString.Split(" ", &Left, &Right);
+
+		double Number = FCString::Atoi(*Left);
+		uint32 IntNumber;
+		uint32 Multiplier = 1;
+
+
+		for (; Multiplier <= 10000; Multiplier *= 10)
 		{
-			UE_LOG(LogCaptureExtractTimecode, Display, TEXT("Failed to run the command: %s %s"), *InFFProbePath, *CommandArgs);
+			double TmpNumber = Number * Multiplier;
+			IntNumber = FMath::RoundToInt32(TmpNumber);
+			if (FMath::Abs(TmpNumber - IntNumber) < 0.01)
+			{
+				break;
+			}
+		}
 
-			FString CommandOutputStr = FString::ConstructFromPtrSize(reinterpret_cast<const UTF8CHAR*>(FullCommandOutput.GetData()), FullCommandOutput.Num());
-			UE_LOG(LogCaptureExtractTimecode, Display,
-				TEXT("Output from the command:\n>>>>>>\n%s<<<<<<"), *CommandOutputStr);
+		uint32 Nominator = 0;
+		uint32 Denominator = 0;
+
+		if (Right == "fps")
+		{
+			Nominator = IntNumber;
+			Denominator = Multiplier;
+		}
+		else if (Right == "s")
+		{
+			Nominator = Multiplier;
+			Denominator = IntNumber;
 		}
 
-		return MakeError(ECaptureExtractTimecodeError::InternalError);
-	}
+		TimecodeRate = FFrameRate(Nominator, Denominator);
 
-	return Result;
+		return MakeValue(TimecodeRate);
+	}
 }
 
 FCaptureExtractAudioTimecode::FCaptureExtractAudioTimecode(const FString& InFilePath)
 	: FilePath(InFilePath)
 {}
 
-UE::CaptureManager::FTimecodeAndRateResult FCaptureExtractAudioTimecode::Extract()
+FCaptureExtractAudioTimecode::FTimecodeInfoResult FCaptureExtractAudioTimecode::Extract()
 {
 	return ExtractTimecodeFromBroadcastWaveFormat(FFrameRate());
 }
@@ -446,13 +498,13 @@ static FFrameRate EstimateSmpteTimecodeRate(const FFrameRate InMediaFrameRate)
 	return InMediaFrameRate;
 }
 
-UE::CaptureManager::FTimecodeAndRateResult FCaptureExtractAudioTimecode::Extract(FFrameRate InFrameRate)
+FCaptureExtractAudioTimecode::FTimecodeInfoResult FCaptureExtractAudioTimecode::Extract(FFrameRate InFrameRate)
 {
 	using namespace UE::CaptureManager;
 
-	check(!FilePath.IsEmpty())
+	check(!FilePath.IsEmpty());
 
-	FTimecodeAndRateResult Result = MakeError(ECaptureExtractTimecodeError::UnhandledMedia);
+	FTimecodeInfoResult Result = MakeError(ECaptureExtractInfoError::UnhandledMedia);
 
 	FString FileExtension = FPaths::GetExtension(FilePath);
 	if (FileExtension == "wav")
@@ -474,7 +526,7 @@ UE::CaptureManager::FTimecodeAndRateResult FCaptureExtractAudioTimecode::Extract
 	return Result;
 }
 
-UE::CaptureManager::FTimecodeAndRateResult FCaptureExtractAudioTimecode::ExtractTimecodeFromBroadcastWaveFormat(FFrameRate InTimecodeRate)
+FCaptureExtractAudioTimecode::FTimecodeInfoResult FCaptureExtractAudioTimecode::ExtractTimecodeFromBroadcastWaveFormat(FFrameRate InTimecodeRate)
 {
 	using namespace UE::CaptureManager;
 
@@ -523,13 +575,18 @@ UE::CaptureManager::FTimecodeAndRateResult FCaptureExtractAudioTimecode::Extract
 				}
 
 				FTimecode AudioTimecode = FTimecode(NumSecondsSinceMidnight, TimecodeRate, TimecodeInfo.bTimecodeIsDropFrame, /* InbRollover = */ true);
-				FTimecodeAndRate TimecodeAndRate{ AudioTimecode, TimecodeRate };
-				return MakeValue(TimecodeAndRate);
+				FTimecodeInfo TimecodeAndRate
+				{ 
+					.Timecode = MoveTemp(AudioTimecode), 
+					.TimecodeRate = MoveTemp(TimecodeRate)
+				};
+
+				return MakeValue(MoveTemp(TimecodeAndRate));
 			}
 		}
 	}
 
-	return MakeError(ECaptureExtractTimecodeError::TimecodeNotFound);
+	return MakeError(ECaptureExtractInfoError::TimecodeNotFound);
 }
 
 #undef LOCTEXT_NAMESPACE
diff --git a/Engine/Plugins/VirtualProduction/CaptureManager/CaptureManagerApp/Source/IngestLiveLinkDevice/Public/Utils/CaptureExtractTimecode.h b/Engine/Plugins/VirtualProduction/CaptureManager/CaptureManagerApp/Source/IngestLiveLinkDevice/Public/Utils/CaptureExtractTimecode.h
--- a/Engine/Plugins/VirtualProduction/CaptureManager/CaptureManagerApp/Source/IngestLiveLinkDevice/Public/Utils/CaptureExtractTimecode.h
+++ b/Engine/Plugins/VirtualProduction/CaptureManager/CaptureManagerApp/Source/IngestLiveLinkDevice/Public/Utils/CaptureExtractTimecode.h
@@ -15,15 +15,24 @@
 
 #include "CaptureExtractTimecode.generated.h"
 
+#define UE_API INGESTLIVELINKDEVICE_API
+
 namespace UE::CaptureManager
 {
-struct FTimecodeAndRate
+
+struct FTimecodeInfo
 {
 	FTimecode Timecode;
-	TOptional<FFrameRate> MaybeTimecodeRate;
+	FFrameRate TimecodeRate;
 };
 
-enum ECaptureExtractTimecodeError : int32
+struct FVideoInformation
+{
+	FFrameRate FrameRate;
+	FTimecodeInfo TimecodeInfo;
+};
+
+enum ECaptureExtractInfoError : int32
 {
 	InternalError = -1,
 	TimecodeNotFound = 1,
@@ -33,8 +42,6 @@ enum ECaptureExtractTimecodeError : int32
 	UnhandledMedia = 5,
 	UnableToOpenMedia = 6,
 };
-
-using FTimecodeAndRateResult = TValueOrError<FTimecodeAndRate, ECaptureExtractTimecodeError>;
 }
 
 #if WITH_EDITOR
@@ -51,35 +58,46 @@ public:
 };
 #endif // WITH_EDITOR
 
-class INGESTLIVELINKDEVICE_API FCaptureExtractVideoTimecode final : public TSharedFromThis<FCaptureExtractVideoTimecode>
+class FCaptureExtractVideoInfo final
 {
 public:
 
-	FCaptureExtractVideoTimecode(const FString& InFilePath);
+	using FResult = TValueOrError<FCaptureExtractVideoInfo, UE::CaptureManager::ECaptureExtractInfoError>;
+
+	UE_API static FResult Create(const FString& InFilePath);
 
-	UE::CaptureManager::FTimecodeAndRateResult Extract();
+	UE_API FFrameRate GetFrameRate() const;
+	UE_API FTimecode GetTimecode() const;
+	UE_API FFrameRate GetTimecodeRate() const;
 
 private:
 
-	FString FilePath;
+	FCaptureExtractVideoInfo();
 
-	TValueOrError<FFrameRate, UE::CaptureManager::ECaptureExtractTimecodeError> ParseTimecodeRate(FString TimecodeRateString);
+	using FExtractResult = TValueOrError<void, UE::CaptureManager::ECaptureExtractInfoError>;
+	FExtractResult ExtractInfo(const FString& InFilePath);
+	FExtractResult ExtractInfoUsingElectraPlayer(const FString& InFilePath);
+	FExtractResult ExtractInfoUsingFFProbe(const FString& InFilePath, const FString& InEncoderPath);
 
-	UE::CaptureManager::FTimecodeAndRateResult ExtractTimecodeUsingElectraPlayer();
-	UE::CaptureManager::FTimecodeAndRateResult ExtractTimecodeUsingFFProbe(const FString& InEncoderPath);
-	
-	const int32 TimeoutPeriod = 3;
+	TValueOrError<FFrameRate, UE::CaptureManager::ECaptureExtractInfoError> ParseTimecodeRate(FString TimecodeRateString);
+
+	static constexpr int32 TimeoutPeriod = 3;
+
+	FFrameRate FrameRate;
+	UE::CaptureManager::FTimecodeInfo TimecodeInfo;
 };
 
-class INGESTLIVELINKDEVICE_API FCaptureExtractAudioTimecode final : public TSharedFromThis<FCaptureExtractAudioTimecode>
+class FCaptureExtractAudioTimecode final : public TSharedFromThis<FCaptureExtractAudioTimecode>
 {
 public:
 
-	FCaptureExtractAudioTimecode(const FString& InFilePath);
-	~FCaptureExtractAudioTimecode() = default;
+	using FTimecodeInfoResult = TValueOrError<UE::CaptureManager::FTimecodeInfo, UE::CaptureManager::ECaptureExtractInfoError>;
+
+	UE_API FCaptureExtractAudioTimecode(const FString& InFilePath);
+	UE_API ~FCaptureExtractAudioTimecode() = default;
 
-	UE::CaptureManager::FTimecodeAndRateResult Extract();
-	UE::CaptureManager::FTimecodeAndRateResult Extract(FFrameRate InFrameRate);
+	UE_API FTimecodeInfoResult Extract();
+	UE_API FTimecodeInfoResult Extract(FFrameRate InFrameRate);
 
 private:
 
@@ -87,5 +105,7 @@ private:
 
 	const int32 TimeoutPeriod = 3;
 
-	UE::CaptureManager::FTimecodeAndRateResult ExtractTimecodeFromBroadcastWaveFormat(FFrameRate InTimecodeRate);
+	FTimecodeInfoResult ExtractTimecodeFromBroadcastWaveFormat(FFrameRate InTimecodeRate);
 };
+
+#undef UE_API
\ No newline at end of file
diff --git a/Engine/Plugins/VirtualProduction/CaptureManager/CaptureManagerDevices/Source/CPSLiveLinkDevice/Private/LiveLinkFaceDevice.cpp b/Engine/Plugins/VirtualProduction/CaptureManager/CaptureManagerDevices/Source/CPSLiveLinkDevice/Private/LiveLinkFaceDevice.cpp
--- a/Engine/Plugins/VirtualProduction/CaptureManager/CaptureManagerDevices/Source/CPSLiveLinkDevice/Private/LiveLinkFaceDevice.cpp
+++ b/Engine/Plugins/VirtualProduction/CaptureManager/CaptureManagerDevices/Source/CPSLiveLinkDevice/Private/LiveLinkFaceDevice.cpp
@@ -572,12 +572,15 @@ void ULiveLinkFaceDevice::ExtractTimecodeIfNotSet(FTakeMetadata& InOutTakeMetada
 	{
 		if (!Video.TimecodeStart.IsSet())
 		{
-			TSharedPtr<FCaptureExtractVideoTimecode> Extractor = MakeShareable(new FCaptureExtractVideoTimecode(Video.Path));
-			FTimecodeAndRateResult TimecodeAndRateResult = Extractor->Extract();
-			if (TimecodeAndRateResult.IsValid())
+			FCaptureExtractVideoInfo::FResult ExtractorOpt = FCaptureExtractVideoInfo::Create(Video.Path);
+			if (ExtractorOpt.IsValid())
 			{
-				FTimecodeAndRate TimecodeAndRate = TimecodeAndRateResult.GetValue();
-				Video.TimecodeStart = TimecodeAndRate.Timecode.ToString();
+				FCaptureExtractVideoInfo Extractor = ExtractorOpt.StealValue();
+				Video.TimecodeStart = Extractor.GetTimecode().ToString();
+				if (FMath::IsNearlyZero(Video.FrameRate))
+				{
+					Video.FrameRate = static_cast<float>(Extractor.GetFrameRate().AsDecimal());
+				}
 			}
 		}
 
@@ -594,16 +597,12 @@ void ULiveLinkFaceDevice::ExtractTimecodeIfNotSet(FTakeMetadata& InOutTakeMetada
 			TSharedPtr<FCaptureExtractAudioTimecode> Extractor = MakeShareable(new FCaptureExtractAudioTimecode(Audio.Path));
 
 			// The video frame rate will be used to calculate the timecode rate if the timecode rate cannot be extracted from the audio file
-			FTimecodeAndRateResult TimecodeAndRateResult = Extractor->Extract(VideoFrameRate); 
-			if (TimecodeAndRateResult.IsValid())
+			FCaptureExtractAudioTimecode::FTimecodeInfoResult TimecodeInfoResult = Extractor->Extract(VideoFrameRate);
+			if (TimecodeInfoResult.IsValid())
 			{
-				FTimecodeAndRate TimecodeAndRate = TimecodeAndRateResult.GetValue();
-				Audio.TimecodeStart = TimecodeAndRate.Timecode.ToString();
-
-				if (TimecodeAndRate.MaybeTimecodeRate.IsSet())
-				{
-					Audio.TimecodeRate = static_cast<float>(TimecodeAndRate.MaybeTimecodeRate.GetValue().AsDecimal());
-				}
+				FTimecodeInfo TimecodeInfo = TimecodeInfoResult.GetValue();
+				Audio.TimecodeStart = TimecodeInfo.Timecode.ToString();
+				Audio.TimecodeRate = static_cast<float>(TimecodeInfo.TimecodeRate.AsDecimal());
 			}
 		}
 	}
diff --git a/Engine/Plugins/VirtualProduction/CaptureManager/CaptureManagerDevices/Source/MonoVideoIngestDevice/Private/MonoVideoIngestDevice.cpp b/Engine/Plugins/VirtualProduction/CaptureManager/CaptureManagerDevices/Source/MonoVideoIngestDevice/Private/MonoVideoIngestDevice.cpp
--- a/Engine/Plugins/VirtualProduction/CaptureManager/CaptureManagerDevices/Source/MonoVideoIngestDevice/Private/MonoVideoIngestDevice.cpp
+++ b/Engine/Plugins/VirtualProduction/CaptureManager/CaptureManagerDevices/Source/MonoVideoIngestDevice/Private/MonoVideoIngestDevice.cpp
@@ -314,12 +314,15 @@ TOptional<FTakeMetadata> UMonoVideoIngestDevice::ReadTake(const FString& InCurre
 	Video.Path = InCurrentTakeFile;
 	Video.PathType = FTakeMetadata::FVideo::EPathType::File;
 
-	TSharedPtr<FCaptureExtractVideoTimecode> Extractor = MakeShareable(new FCaptureExtractVideoTimecode(Video.Path));
-	FTimecodeAndRateResult TimecodeAndRateResult = Extractor->Extract();
-	if (TimecodeAndRateResult.IsValid())
+	FCaptureExtractVideoInfo::FResult ExtractorOpt = FCaptureExtractVideoInfo::Create(Video.Path);
+	if (ExtractorOpt.IsValid())
 	{
-		FTimecodeAndRate TimecodeAndRate = TimecodeAndRateResult.GetValue();
-		Video.TimecodeStart = TimecodeAndRate.Timecode.ToString();
+		FCaptureExtractVideoInfo Extractor = ExtractorOpt.StealValue();
+		Video.TimecodeStart = Extractor.GetTimecode().ToString();
+		if (FMath::IsNearlyZero(Video.FrameRate))
+		{
+			Video.FrameRate = static_cast<float>(Extractor.GetFrameRate().AsDecimal());
+		}
 	}
 
 	TakeMetadata.Video.Add(MoveTemp(Video));
diff --git a/Engine/Plugins/VirtualProduction/CaptureManager/CaptureManagerDevices/Source/StereoVideoIngestDevice/Private/StereoVideoIngestDevice.cpp b/Engine/Plugins/VirtualProduction/CaptureManager/CaptureManagerDevices/Source/StereoVideoIngestDevice/Private/StereoVideoIngestDevice.cpp
--- a/Engine/Plugins/VirtualProduction/CaptureManager/CaptureManagerDevices/Source/StereoVideoIngestDevice/Private/StereoVideoIngestDevice.cpp
+++ b/Engine/Plugins/VirtualProduction/CaptureManager/CaptureManagerDevices/Source/StereoVideoIngestDevice/Private/StereoVideoIngestDevice.cpp
@@ -538,12 +538,16 @@ TOptional<FTakeMetadata> UStereoVideoIngestDevice::CreateTakeMetadata(const FTak
 				Video.Path = TakeComponent.Path;
 				
 				FString MediaFilePath = FPaths::ConvertRelativePathToFull(InTakeWithComponents.TakePath, Video.Path);
-				TSharedPtr<FCaptureExtractVideoTimecode> Extractor = MakeShareable(new FCaptureExtractVideoTimecode(MediaFilePath));
-				FTimecodeAndRateResult TimecodeAndRateResult = Extractor->Extract();
-				if (TimecodeAndRateResult.IsValid())
+
+				FCaptureExtractVideoInfo::FResult ExtractorOpt = FCaptureExtractVideoInfo::Create(MediaFilePath);
+				if (ExtractorOpt.IsValid())
 				{
-					FTimecodeAndRate TimecodeAndRate = TimecodeAndRateResult.GetValue();
-					Video.TimecodeStart = TimecodeAndRate.Timecode.ToString();
+					FCaptureExtractVideoInfo Extractor = ExtractorOpt.StealValue();
+					Video.TimecodeStart = Extractor.GetTimecode().ToString();
+					if (FMath::IsNearlyZero(Video.FrameRate))
+					{
+						Video.FrameRate = static_cast<float>(Extractor.GetFrameRate().AsDecimal());
+					}
 				}
 			}
 			else
@@ -610,16 +614,13 @@ TOptional<FTakeMetadata> UStereoVideoIngestDevice::CreateTakeMetadata(const FTak
 			TSharedPtr<FCaptureExtractAudioTimecode> Extractor = MakeShareable(new FCaptureExtractAudioTimecode(MediaFilePath));
 
 			// The video frame rate will be used to calculate the timecode rate if the timecode rate cannot be extracted from the audio file
-			FTimecodeAndRateResult TimecodeAndRateResult = Extractor->Extract(VideoFrameRate);
-			if (TimecodeAndRateResult.IsValid())
+			FCaptureExtractAudioTimecode::FTimecodeInfoResult TimecodeInfoResult = Extractor->Extract(VideoFrameRate);
+			if (TimecodeInfoResult.IsValid())
 			{
-				FTimecodeAndRate TimecodeAndRate = TimecodeAndRateResult.GetValue();
-				Audio.TimecodeStart = TimecodeAndRate.Timecode.ToString();
-			
-				if (TimecodeAndRate.MaybeTimecodeRate.IsSet())
-				{
-					Audio.TimecodeRate = static_cast<float>(TimecodeAndRate.MaybeTimecodeRate.GetValue().AsDecimal());
-				}
+				FTimecodeInfo TimecodeInfo = TimecodeInfoResult.GetValue();
+
+				Audio.TimecodeStart = TimecodeInfo.Timecode.ToString();
+				Audio.TimecodeRate = static_cast<float>(TimecodeInfo.TimecodeRate.AsDecimal());
 			}
 
 			TakeMetadata.Audio.Add(MoveTemp(Audio));
diff --git a/Engine/Plugins/VirtualProduction/CaptureManager/CaptureManagerDevices/Source/TakeArchiveIngestDevice/Private/TakeArchiveIngestDevice.cpp b/Engine/Plugins/VirtualProduction/CaptureManager/CaptureManagerDevices/Source/TakeArchiveIngestDevice/Private/TakeArchiveIngestDevice.cpp
--- a/Engine/Plugins/VirtualProduction/CaptureManager/CaptureManagerDevices/Source/TakeArchiveIngestDevice/Private/TakeArchiveIngestDevice.cpp
+++ b/Engine/Plugins/VirtualProduction/CaptureManager/CaptureManagerDevices/Source/TakeArchiveIngestDevice/Private/TakeArchiveIngestDevice.cpp
@@ -254,12 +254,15 @@ TOptional<FTakeMetadata> UTakeArchiveIngestDevice::ReadTake(const FString& InCur
 		// Check if this is a file (and not an image sequence directory)
 		if (FPaths::FileExists(Video.Path) && !Video.TimecodeStart.IsSet())
 		{
-			TSharedPtr<FCaptureExtractVideoTimecode> Extractor = MakeShareable(new FCaptureExtractVideoTimecode(Video.Path));
-			FTimecodeAndRateResult TimecodeAndRateResult = Extractor->Extract();
-			if (TimecodeAndRateResult.IsValid())
+			FCaptureExtractVideoInfo::FResult ExtractorOpt = FCaptureExtractVideoInfo::Create(Video.Path);
+			if (ExtractorOpt.IsValid())
 			{
-				FTimecodeAndRate TimecodeAndRate = TimecodeAndRateResult.GetValue();
-				Video.TimecodeStart = TimecodeAndRate.Timecode.ToString();
+				FCaptureExtractVideoInfo Extractor = ExtractorOpt.StealValue();
+				Video.TimecodeStart = Extractor.GetTimecode().ToString();
+				if (FMath::IsNearlyZero(Video.FrameRate))
+				{
+					Video.FrameRate = static_cast<float>(Extractor.GetFrameRate().AsDecimal());
+				}
 			}
 		}
 
@@ -280,16 +283,13 @@ TOptional<FTakeMetadata> UTakeArchiveIngestDevice::ReadTake(const FString& InCur
 
 		// The video frame rate will be used to calculate the timecode rate if the timecode rate cannot be extracted from the audio file
 		// If the VideoFrameRate is not set, the timecode rate may be set to the audio sample rate
-		FTimecodeAndRateResult TimecodeAndRateResult = Extractor->Extract(VideoFrameRate);
-		if (TimecodeAndRateResult.IsValid())
+		FCaptureExtractAudioTimecode::FTimecodeInfoResult TimecodeInfoResult = Extractor->Extract(VideoFrameRate);
+		if (TimecodeInfoResult.IsValid())
 		{
-			FTimecodeAndRate TimecodeAndRate = TimecodeAndRateResult.GetValue();
-			Audio.TimecodeStart = TimecodeAndRate.Timecode.ToString();
-			
-			if (TimecodeAndRate.MaybeTimecodeRate.IsSet())
-			{
-				Audio.TimecodeRate = static_cast<float>(TimecodeAndRate.MaybeTimecodeRate.GetValue().AsDecimal());
-			}
+			FTimecodeInfo TimecodeInfo = TimecodeInfoResult.GetValue();
+
+			Audio.TimecodeStart = TimecodeInfo.Timecode.ToString();
+			Audio.TimecodeRate = static_cast<float>(TimecodeInfo.TimecodeRate.AsDecimal());
 		}
 	}
 	
