Update frame log based timecode parsing to accept video frame rates other than 30 and 60 FPS.

Also permit depth frame rates other than 30 FPS. Maintain the audio timecode generation to use 30 FPS for now.

diff --git a/Engine/Plugins/VirtualProduction/CaptureManager/CaptureManagerApp/Source/LiveLinkFaceMetadata/Private/LiveLinkFaceMetadata.cpp b/Engine/Plugins/VirtualProduction/CaptureManager/CaptureManagerApp/Source/LiveLinkFaceMetadata/Private/LiveLinkFaceMetadata.cpp
--- a/Engine/Plugins/VirtualProduction/CaptureManager/CaptureManagerApp/Source/LiveLinkFaceMetadata/Private/LiveLinkFaceMetadata.cpp
+++ b/Engine/Plugins/VirtualProduction/CaptureManager/CaptureManagerApp/Source/LiveLinkFaceMetadata/Private/LiveLinkFaceMetadata.cpp
@@ -782,20 +782,18 @@ class FFrameLogParser
 {
 public:
 
-	void ParseFrameLog(const FString& InFrameLogPath, float InFrameRate)
+	void ParseFrameLog(const FString& InFrameLogPath, float InVideoFrameRate, float InDepthFrameRate)
 	{
 		TArray<FString> FrameLogLines;
 		FFileHelper::LoadFileToStringArray(FrameLogLines, *InFrameLogPath);
 
 		bool bAudioTimecodeFound = false;
-		FFrameRate OriginalFrameRate;
-		if (FMath::IsNearlyEqual(InFrameRate, 30.0))
-		{
-			OriginalFrameRate = FFrameRate{ 30, 1 };
-		}
-		else if (FMath::IsNearlyEqual(InFrameRate, 60.0))
+
+		FFrameRate OriginalVideoFrameRate;
+		uint32 RoundedVideoFrameRate = static_cast<uint32>(FMath::RoundToInt(InVideoFrameRate));
+		if (RoundedVideoFrameRate != 0)
 		{
-			OriginalFrameRate = FFrameRate{ 60, 1 };
+			OriginalVideoFrameRate = FFrameRate{ RoundedVideoFrameRate, 1 };
 		}
 
 		for (const FString& Line : FrameLogLines)
@@ -817,7 +815,7 @@ public:
 					}
 
 					// If frame rate isn't known, there's no sense to parse timecode
-					if (OriginalFrameRate == FFrameRate{})
+					if (OriginalVideoFrameRate == FFrameRate{})
 					{
 						continue;
 					}
@@ -826,13 +824,13 @@ public:
 					// which may have occurred while parsing the frame log. As an example: a fractional frame number of
 					// 59.780 would round to 60, which is an invalid frame number at 60FPS. Going to/from a frame number
 					// fixes the problem.
-					VideoTimecode = FTimecode::FromFrameNumber(Timecode.ToFrameNumber(OriginalFrameRate), OriginalFrameRate);
-					VideoTimecodeRate = OriginalFrameRate;
+					VideoTimecode = FTimecode::FromFrameNumber(Timecode.ToFrameNumber(OriginalVideoFrameRate), OriginalVideoFrameRate);
+					VideoTimecodeRate = OriginalVideoFrameRate;
 				}
 			}
 			else if (LogEntry.EntryType() == FrameLogEntry::DepthType)
 			{
-				if (LogEntry.FrameIndex() == 0) // Take timecode from first video frame
+				if (LogEntry.FrameIndex() == 0) // Take timecode from first depth frame
 				{
 					FTimecode Timecode;
 					if (!LogEntry.Timecode(Timecode))
@@ -840,13 +838,19 @@ public:
 						continue;
 					}
 
-					// Depth timecode rate and timecode are 30fps, and there is no need for conversion
-					FFrameRate TargetFrameRate{ 30, 1 };
+					uint32 RoundedDepthFrameRate = static_cast<uint32>(FMath::RoundToInt(InDepthFrameRate));
+					if (RoundedDepthFrameRate == 0)
+					{
+						continue;
+					}
+
+					FFrameRate OriginalDepthFrameRate = FFrameRate{ RoundedDepthFrameRate, 1 };
 
 					// This step is also needed to resolve any invalid frame number rounding which may have occurred
 					// during the parsing of the frame log (as it implicitly does a conversion to/from frame number).
-					DepthTimecode = FTimecode::FromTimespan(Timecode.ToTimespan(OriginalFrameRate), TargetFrameRate, true);
-					DepthTimecodeRate = MoveTemp(TargetFrameRate);
+					// The depth timecode in the frame log matches the frame rate of the video frame rate so we use the video frame rate when creating the timespan
+					DepthTimecode = FTimecode::FromTimespan(Timecode.ToTimespan(OriginalVideoFrameRate), OriginalDepthFrameRate, true);
+					DepthTimecodeRate = OriginalDepthFrameRate;
 				}
 
 				++DepthFrameCount;
@@ -862,19 +866,20 @@ public:
 					}
 
 					// If frame rate isn't known, there's no sense to parse timecode
-					if (OriginalFrameRate == FFrameRate{})
+					if (OriginalVideoFrameRate == FFrameRate{})
 					{
 						continue;
 					}
 
-					// Make sure resulting time code is 30fps to match the depth which is (currently always 30fps).
-					// This probably isn't mandatory and we should at some point step away from this when we make
-					// sure this assumption is not present in other parts of the MHA codebase.
+					// Originally we intentionally used the same frame rate here as for depth (assuming depth at 30 FPS), but now depth frame rate can vary.
+					// In the near future we could consider matching the video frame rate here instead, but for the sake of
+					// reducing risk for 5.7 we will maintain the current behaviour.
 					FFrameRate TargetFrameRate{ 30, 1 };
 
+					// The audio timecode is expressed at the video frame rate 
 					// This step is also needed to resolve any invalid frame number rounding which may have occurred
 					// during the parsing of the frame log (as it implicitly does a conversion to/from frame number).
-					AudioTimecode = FTimecode::FromTimespan(Timecode.ToTimespan(OriginalFrameRate), TargetFrameRate, true);
+					AudioTimecode = FTimecode::FromTimespan(Timecode.ToTimespan(OriginalVideoFrameRate), TargetFrameRate, true);
 					AudioTimecodeRate = TargetFrameRate;
 					bAudioTimecodeFound = true;
 				}
@@ -949,7 +954,7 @@ FTakeMetadata ConvertOldToNewTakeMetadata(const FLiveLinkFaceTakeInfo& InTakeInf
 	NewTakeInfo.Device.Platform = MoveTemp(Platform);
 
 	FFrameLogParser Parser;
-	Parser.ParseFrameLog(InTakeInfo.GetFrameLogFilePath(), InTakeInfo.VideoMetadata.FrameRate);
+	Parser.ParseFrameLog(InTakeInfo.GetFrameLogFilePath(), InTakeInfo.VideoMetadata.FrameRate, InTakeInfo.DepthMetadata.FrameRate);
 
 	FTakeMetadata::FVideo Video;
 	Video.Name = "Video";
