[Capture Manager] Fixed several crashes while aborting take upload.

There were a few thread safety issues that needed addressing here. There are lots more improvements to be made but this should make things stable. In essence what has happened here is that the old FCommunicationManager and FUploader classes have been removed and their responsibilities absorbed into the endpoint.

[RN] Virtual Production Tools

diff --git a/Engine/Plugins/VirtualProduction/CaptureManager/CaptureManagerApp/Source/CaptureManagerEditor/Private/IngestManagement/IngestJobExecutor.cpp b/Engine/Plugins/VirtualProduction/CaptureManager/CaptureManagerApp/Source/CaptureManagerEditor/Private/IngestManagement/IngestJobExecutor.cpp
--- a/Engine/Plugins/VirtualProduction/CaptureManager/CaptureManagerApp/Source/CaptureManagerEditor/Private/IngestManagement/IngestJobExecutor.cpp
+++ b/Engine/Plugins/VirtualProduction/CaptureManager/CaptureManagerApp/Source/CaptureManagerEditor/Private/IngestManagement/IngestJobExecutor.cpp
@@ -54,7 +54,29 @@ static EIngestCapability_ImageRotation ConvertImageRotation(::EImageRotation InI
 	}
 }
 
-FIngestJobExecutor::FIngestJobExecutor(FString InExecutorName, TSharedRef<FIngestJobQueue> InProcessingQueue, FOnComplete InOnComplete, FJobProcessingStateChanged InJobProcessingStateChanged) :
+TSharedRef<FIngestJobExecutor> FIngestJobExecutor::Create(
+	FString InExecutorName,
+	TSharedRef<FIngestJobQueue> InProcessingQueue,
+	FOnComplete InOnComplete,
+	FJobProcessingStateChanged InJobProcessingStateChanged
+)
+{
+	return MakeShared<FIngestJobExecutor>(
+		FPrivateToken(),
+		MoveTemp(InExecutorName),
+		MoveTemp(InProcessingQueue),
+		MoveTemp(InOnComplete),
+		MoveTemp(InJobProcessingStateChanged)
+	);
+}
+
+FIngestJobExecutor::FIngestJobExecutor(
+	FPrivateToken InPrivateToken,
+	FString InExecutorName,
+	TSharedRef<FIngestJobQueue> InProcessingQueue,
+	FOnComplete InOnComplete,
+	FJobProcessingStateChanged InJobProcessingStateChanged
+) :
 	ExecutorName(MoveTemp(InExecutorName)),
 	ProcessingQueue(MoveTemp(InProcessingQueue)),
 	Complete(MoveTemp(InOnComplete)),
@@ -83,8 +105,8 @@ uint32 FIngestJobExecutor::Run()
 {
 	while (!bStopRequested)
 	{
-		TPromise<void> Promise;
-		TFuture<void> Future = Promise.GetFuture();
+		TSharedRef<TPromise<void>> Promise = MakeShared<TPromise<void>>();
+		TFuture<void> Future = Promise->GetFuture();
 
 		TSharedPtr<FIngestJob> Job = ProcessingQueue->ClaimFirstPending(FIngestJob::EProcessingState::Running);
 
@@ -96,9 +118,15 @@ uint32 FIngestJobExecutor::Run()
 
 			FScopeLock Lock(&CriticalSection);
 
+			// RunIngest MUST set the promise for us to terminate cleanly, so we pass shared references to keep things alive in
+			// the various callbacks inside RunIngest, just to help ensure this. Weak ptrs or references etc. would open up
+			// possibilities (every time you try to .Pin()) for that not to be the case. Since these callbacks are a part of
+			// the process handle inside the context object (which is set as a member of this class), everything should be well
+			// within our control.
+
 			// Note: This may be a blocking call depending on how the capture device has been implemented. If so, this 
 			// may take a while to stop (no way to cancel in that case).
-			TUniquePtr<FIngestJobExecutor::FRunIngestContext> RunIngestContext = RunIngest(*Job, Promise);
+			TUniquePtr<FIngestJobExecutor::FRunIngestContext> RunIngestContext = RunIngest(Job.ToSharedRef(), Promise);
 			
 			// Check the context in case the device has been removed during the RunIngest call
 			if (RunIngestContext)
@@ -116,7 +144,7 @@ uint32 FIngestJobExecutor::Run()
 		else
 		{
 			// Job queue depleted, stop the executor
-			Promise.SetValue();
+			Promise->SetValue();
 			bStopRequested = true;
 		}
 
@@ -168,8 +196,10 @@ void FIngestJobExecutor::Exit()
 	Complete.ExecuteIfBound();
 }
 
-TUniquePtr<FIngestJobExecutor::FRunIngestContext> FIngestJobExecutor::RunIngest(FIngestJob& InJob, TPromise<void>& InPromise)
+TUniquePtr<FIngestJobExecutor::FRunIngestContext> FIngestJobExecutor::RunIngest(TSharedRef<FIngestJob> InJob, TSharedRef<TPromise<void>> InPromise)
 {
+	TRACE_CPUPROFILER_EVENT_SCOPE(FIngestJobExecutor_RunIngest);
+
 	if (!GEditor)
 	{
 		FText Message = LOCTEXT("IngestProcessStateMessage_GeditorUnavailable", "Ingest failed, GEditor unavailble");
@@ -180,15 +210,15 @@ TUniquePtr<FIngestJobExecutor::FRunIngestContext> FIngestJobExecutor::RunIngest(
 		);
 
 		// Set the processing state to anything other than Pending, to prevent other executors picking up the same bad job.
-		InJob.SetProcessingState(FIngestJob::FProcessingState{ FIngestJob::EProcessingState::Aborted, MoveTemp(Message) });
-		InPromise.SetValue();
+		InJob->SetProcessingState(FIngestJob::FProcessingState{ FIngestJob::EProcessingState::Aborted, MoveTemp(Message) });
+		InPromise->SetValue();
 		return nullptr;
 	}
 
 	ULiveLinkDeviceSubsystem* DeviceSubsystem = GEngine->GetEngineSubsystem<ULiveLinkDeviceSubsystem>();
 	check(DeviceSubsystem);
 
-	const TObjectPtr<ULiveLinkDevice>* DevicePtr = DeviceSubsystem->GetDeviceMap().Find(InJob.GetCaptureDeviceId());
+	const TObjectPtr<ULiveLinkDevice>* DevicePtr = DeviceSubsystem->GetDeviceMap().Find(InJob->GetCaptureDeviceId());
 
 	if (!DevicePtr)
 	{
@@ -200,8 +230,8 @@ TUniquePtr<FIngestJobExecutor::FRunIngestContext> FIngestJobExecutor::RunIngest(
 		);
 
 		// Set the processing state to anything other than Pending, to prevent other executors picking up the same bad job.
-		InJob.SetProcessingState(FIngestJob::FProcessingState{ FIngestJob::EProcessingState::Aborted, MoveTemp(Message) });
-		InPromise.SetValue();
+		InJob->SetProcessingState(FIngestJob::FProcessingState{ FIngestJob::EProcessingState::Aborted, MoveTemp(Message) });
+		InPromise->SetValue();
 		return nullptr;
 	}
 
@@ -224,17 +254,17 @@ TUniquePtr<FIngestJobExecutor::FRunIngestContext> FIngestJobExecutor::RunIngest(
 	GeneralNamingTokenArgs.AdditionalNamespacesToInclude.Add(Tokens->GetNamespace());
 	GeneralNamingTokenArgs.bNativeOnly = true;
 
-	const FIngestJob::FSettings Settings = InJob.GetSettings();
+	const FIngestJob::FSettings Settings = InJob->GetSettings();
 	// If everything is working correctly the upload host name should never be empty.
 	check(!Settings.UploadHostName.IsEmpty());
 
 	FStringFormatNamedArguments GeneralArgs;
 	{
 		using namespace UE::CaptureManager;
-		GeneralArgs.Add(Tokens->GetToken(FString(GeneralTokens::IdKey)).Name, InJob.GetTakeMetadata().UniqueId);
+		GeneralArgs.Add(Tokens->GetToken(FString(GeneralTokens::IdKey)).Name, InJob->GetTakeMetadata().UniqueId);
 		GeneralArgs.Add(Tokens->GetToken(FString(GeneralTokens::DeviceKey)).Name, DeviceName.ToString());
-		GeneralArgs.Add(Tokens->GetToken(FString(GeneralTokens::SlateKey)).Name, InJob.GetTakeMetadata().Slate);
-		GeneralArgs.Add(Tokens->GetToken(FString(GeneralTokens::TakeKey)).Name, InJob.GetTakeMetadata().TakeNumber);
+		GeneralArgs.Add(Tokens->GetToken(FString(GeneralTokens::SlateKey)).Name, InJob->GetTakeMetadata().Slate);
+		GeneralArgs.Add(Tokens->GetToken(FString(GeneralTokens::TakeKey)).Name, InJob->GetTakeMetadata().TakeNumber);
 	}
 
 	FString WorkingDirectory = FString::Format(*Settings.WorkingDirectory, GeneralArgs);
@@ -273,13 +303,13 @@ TUniquePtr<FIngestJobExecutor::FRunIngestContext> FIngestJobExecutor::RunIngest(
 		);
 
 		// Set the processing state to anything other than Pending, to prevent other executors picking up the same bad job.
-		InJob.SetProcessingState(FIngestJob::FProcessingState{FIngestJob::EProcessingState::Aborted, MoveTemp(Message)});
-		InPromise.SetValue();
+		InJob->SetProcessingState(FIngestJob::FProcessingState{FIngestJob::EProcessingState::Aborted, MoveTemp(Message)});
+		InPromise->SetValue();
 		return nullptr;
 	}
 
 	FIngestProcessFinishReporter ProcessFinished = FIngestProcessFinishReporter::Type::CreateLambda(
-		[this, &InPromise, &InJob](const UIngestCapability_ProcessHandle* InProcessHandle, TValueOrError<void, FIngestCapability_Error> IngestError)
+		[This = AsShared(), Promise = InPromise, Job = InJob](const UIngestCapability_ProcessHandle* InProcessHandle, TValueOrError<void, FIngestCapability_Error> IngestError) mutable
 		{
 			FIngestJob::EProcessingState ProcessingState;
 			FText Message;
@@ -290,13 +320,13 @@ TUniquePtr<FIngestJobExecutor::FRunIngestContext> FIngestJobExecutor::RunIngest(
 			{
 				ProcessingState = FIngestJob::EProcessingState::Aborted;
 				Message = FText::Format(LOCTEXT("IngestProcessStateMessage_Aborted", "Ingest failed : {0} (Error code : {1})"), FText::FromString(IngestError.GetError().GetMessage()), FText::AsNumber(IngestError.GetError().GetCode()));
-				UE_LOG(LogCaptureManagerIngestJobExecutor, Error, TEXT("[%s] %s"), *ExecutorName, *Message.ToString());
+				UE_LOG(LogCaptureManagerIngestJobExecutor, Error, TEXT("[%s] %s"), *This->ExecutorName, *Message.ToString());
 				bIsDone = true;
 			}
 			else if (bIsDone)
 			{
 				ProcessingState = FIngestJob::EProcessingState::Complete;
-				InJob.SetProgress(1.0);
+				Job->SetProgress(1.0);
 				Message = LOCTEXT("IngestProcessStateMessage_Complete", "Success");
 			}
 			else
@@ -305,29 +335,29 @@ TUniquePtr<FIngestJobExecutor::FRunIngestContext> FIngestJobExecutor::RunIngest(
 				Message = LOCTEXT("IngestProcessStateMessage_Running", "Running...");
 			}
 
-			InJob.SetProcessingState(FIngestJob::FProcessingState{ ProcessingState, MoveTemp(Message) });
-			JobProcessingStateChanged.ExecuteIfBound(InJob.GetGuid(), ProcessingState);
+			Job->SetProcessingState(FIngestJob::FProcessingState{ ProcessingState, MoveTemp(Message) });
+			This->JobProcessingStateChanged.ExecuteIfBound(Job->GetGuid(), ProcessingState);
 			
 			if (bIsDone)
 			{
-				InPromise.SetValue();
+				Promise->SetValue();
 			}
 		}
 	);
 
 	FIngestProcessProgressReporter ProgressReporter = FIngestProcessProgressReporter::Type::CreateLambda(
-		[&InJob, LastProgress = 0.0](const UIngestCapability_ProcessHandle* InProcessHandle, const double Progress) mutable
+		[Job = InJob, LastProgress = 0.0](const UIngestCapability_ProcessHandle* InProcessHandle, const double Progress) mutable
 		{
 			// A little sanity check that we're receiving updates in the expected order.
 			check(Progress >= LastProgress);
 			LastProgress = Progress;
 
-			InJob.SetProgress(Progress);
+			Job->SetProgress(Progress);
 		}
 	);
 
 	TStrongObjectPtr<UIngestCapability_ProcessHandle> ProcessHandle(
-		ILiveLinkDeviceCapability_Ingest::Execute_CreateIngestProcess(Device, InJob.GetTakeId(), InJob.GetPipelineConfig()));
+		ILiveLinkDeviceCapability_Ingest::Execute_CreateIngestProcess(Device, InJob->GetTakeId(), InJob->GetPipelineConfig()));
 
 	ProcessHandle->OnProcessFinishReporter() = MoveTemp(ProcessFinished);
 	ProcessHandle->OnProcessProgressReporter() = MoveTemp(ProgressReporter);
diff --git a/Engine/Plugins/VirtualProduction/CaptureManager/CaptureManagerApp/Source/CaptureManagerEditor/Private/IngestManagement/IngestJobExecutor.h b/Engine/Plugins/VirtualProduction/CaptureManager/CaptureManagerApp/Source/CaptureManagerEditor/Private/IngestManagement/IngestJobExecutor.h
--- a/Engine/Plugins/VirtualProduction/CaptureManager/CaptureManagerApp/Source/CaptureManagerEditor/Private/IngestManagement/IngestJobExecutor.h
+++ b/Engine/Plugins/VirtualProduction/CaptureManager/CaptureManagerApp/Source/CaptureManagerEditor/Private/IngestManagement/IngestJobExecutor.h
@@ -10,15 +10,38 @@ namespace UE::CaptureManager
 {
 class FIngestJobQueue;
 
-class FIngestJobExecutor : public FRunnable
+class FIngestJobExecutor : public FRunnable, public TSharedFromThis<FIngestJobExecutor>
 {
+private:
+	struct FPrivateToken { explicit FPrivateToken() = default; };
+
 public:
 	DECLARE_DELEGATE(FOnComplete);
 	DECLARE_DELEGATE_TwoParams(FJobProcessingStateChanged, FGuid, FIngestJob::EProcessingState);
 
-	FIngestJobExecutor(FString InExecutorName, TSharedRef<FIngestJobQueue> InProcessingQueue, FOnComplete InOnComplete, FJobProcessingStateChanged InJobProcessingStateChanged);
+	static TSharedRef<FIngestJobExecutor> Create(
+		FString InExecutorName,
+		TSharedRef<FIngestJobQueue> InProcessingQueue,
+		FOnComplete InOnComplete,
+		FJobProcessingStateChanged InJobProcessingStateChanged
+	);
+
+	FIngestJobExecutor(
+		FPrivateToken InPrivateToken,
+		FString InExecutorName,
+		TSharedRef<FIngestJobQueue> InProcessingQueue,
+		FOnComplete InOnComplete,
+		FJobProcessingStateChanged InJobProcessingStateChanged
+	);
+
 	~FIngestJobExecutor();
 
+	FIngestJobExecutor(const FIngestJobExecutor&) = delete;
+	FIngestJobExecutor& operator=(const FIngestJobExecutor&) = delete;
+
+	FIngestJobExecutor(FIngestJobExecutor&&) = delete;
+	FIngestJobExecutor operator=(FIngestJobExecutor&&) = delete;
+
 	virtual bool Init() override;
 	virtual uint32 Run() override;
 	virtual void Stop() override;
@@ -33,7 +56,7 @@ private:
 		TStrongObjectPtr<UIngestCapability_ProcessHandle> ProcessHandle;
 	};
 
-	TUniquePtr<FIngestJobExecutor::FRunIngestContext> RunIngest(FIngestJob& InJob, TPromise<void>& InPromise);
+	TUniquePtr<FIngestJobExecutor::FRunIngestContext> RunIngest(TSharedRef<FIngestJob> InJob, TSharedRef<TPromise<void>> InPromise);
 
 	FCriticalSection CriticalSection;
 	FString ExecutorName;
diff --git a/Engine/Plugins/VirtualProduction/CaptureManager/CaptureManagerApp/Source/CaptureManagerEditor/Private/IngestManagement/IngestJobProcessor.cpp b/Engine/Plugins/VirtualProduction/CaptureManager/CaptureManagerApp/Source/CaptureManagerEditor/Private/IngestManagement/IngestJobProcessor.cpp
--- a/Engine/Plugins/VirtualProduction/CaptureManager/CaptureManagerApp/Source/CaptureManagerEditor/Private/IngestManagement/IngestJobProcessor.cpp
+++ b/Engine/Plugins/VirtualProduction/CaptureManager/CaptureManagerApp/Source/CaptureManagerEditor/Private/IngestManagement/IngestJobProcessor.cpp
@@ -119,7 +119,7 @@ void FIngestJobProcessor::StartProcessing()
 		{
 			FScopeLock Lock(&CriticalSection);
 			Executors.Emplace(
-				MakeUnique<FIngestJobExecutor>(
+				FIngestJobExecutor::Create(
 					FString::Printf(TEXT("Ingest Executor %d"), NumExecutor + 1),
 					ProcessingQueue,
 					MoveTemp(OnComplete),
@@ -141,12 +141,12 @@ void FIngestJobProcessor::StopProcessing()
 
 	FScopeLock Lock(&CriticalSection);
 
-	for (const TUniquePtr<FIngestJobExecutor>& Executor : Executors)
+	for (const TSharedRef<FIngestJobExecutor>& Executor : Executors)
 	{
 		Executor->Stop();
 	}
 
-	// IsProcessing will return true until all the processors have terminated
+	// IsProcessing will return true until all the executors have terminated
 }
 
 bool FIngestJobProcessor::IsProcessing() const
@@ -199,7 +199,7 @@ void FIngestJobProcessor::StopProcessingForDevice(const FGuid& InDeviceId)
 
 	FScopeLock Lock(&CriticalSection);
 
-	for (const TUniquePtr<FIngestJobExecutor>& Executor : Executors)
+	for (const TSharedRef<FIngestJobExecutor>& Executor : Executors)
 	{
 		Executor->CancelForDevice(InDeviceId);
 	}
diff --git a/Engine/Plugins/VirtualProduction/CaptureManager/CaptureManagerApp/Source/CaptureManagerEditor/Private/IngestManagement/IngestJobProcessor.h b/Engine/Plugins/VirtualProduction/CaptureManager/CaptureManagerApp/Source/CaptureManagerEditor/Private/IngestManagement/IngestJobProcessor.h
--- a/Engine/Plugins/VirtualProduction/CaptureManager/CaptureManagerApp/Source/CaptureManagerEditor/Private/IngestManagement/IngestJobProcessor.h
+++ b/Engine/Plugins/VirtualProduction/CaptureManager/CaptureManagerApp/Source/CaptureManagerEditor/Private/IngestManagement/IngestJobProcessor.h
@@ -55,7 +55,7 @@ private:
 
 	const int32 NumExecutors;
 	std::atomic<int32> NumExecutorsRunning = 0;
-	TArray<TUniquePtr<FIngestJobExecutor>> Executors;
+	TArray<TSharedRef<FIngestJobExecutor>> Executors;
 	TSharedRef<FIngestJobQueue> ProcessingQueue;
 	std::atomic<bool> bStopRequested = false;
 
diff --git a/Engine/Plugins/VirtualProduction/CaptureManager/CaptureManagerApp/Source/CaptureManagerUnrealEndpoint/Private/CaptureCommunicationManager.cpp b/Engine/Plugins/VirtualProduction/CaptureManager/CaptureManagerApp/Source/CaptureManagerUnrealEndpoint/Private/CaptureCommunicationManager.cpp
deleted file mode 100644
diff --git a/Engine/Plugins/VirtualProduction/CaptureManager/CaptureManagerApp/Source/CaptureManagerUnrealEndpoint/Private/CaptureCommunicationManager.h b/Engine/Plugins/VirtualProduction/CaptureManager/CaptureManagerApp/Source/CaptureManagerUnrealEndpoint/Private/CaptureCommunicationManager.h
deleted file mode 100644
diff --git a/Engine/Plugins/VirtualProduction/CaptureManager/CaptureManagerApp/Source/CaptureManagerUnrealEndpoint/Private/CaptureManagerUnrealEndpoint.cpp b/Engine/Plugins/VirtualProduction/CaptureManager/CaptureManagerApp/Source/CaptureManagerUnrealEndpoint/Private/CaptureManagerUnrealEndpoint.cpp
--- a/Engine/Plugins/VirtualProduction/CaptureManager/CaptureManagerApp/Source/CaptureManagerUnrealEndpoint/Private/CaptureManagerUnrealEndpoint.cpp
+++ b/Engine/Plugins/VirtualProduction/CaptureManager/CaptureManagerApp/Source/CaptureManagerUnrealEndpoint/Private/CaptureManagerUnrealEndpoint.cpp
@@ -3,52 +3,224 @@
 #include "CaptureManagerUnrealEndpoint.h"
 #include "CaptureManagerUnrealEndpointLog.h"
 
-#include "CaptureCommunicationManager.h"
+#include "Async/Async.h"
+#include "Messenger.h"
+#include "Features/ConnectStarter.h"
+#include "Features/UploadStateHandler.h"
+
+#include "LiveLinkHubExportManager.h"
+
+#include "Tasks/Task.h"
+#include "Templates/Greater.h"
 
 #include <condition_variable>
 
 namespace UE::CaptureManager
 {
 
+// This class just bundles together some admin for tracking and aborting take upload tasks
+class FAugmentedTakeUploadTask
+{
+public:
+	FAugmentedTakeUploadTask(TUniquePtr<FTakeUploadTask> InTakeUploadTask, TWeakPtr<FLiveLinkHubExportManager> InExportManager);
+
+	bool Start(FLiveLinkHubExportClient::FTakeUploadParams InUploadParams);
+	bool Abort();
+	FGuid GetTaskID() const;
+
+	// Return a weak ptr, just to indicate that the lifetime of the upload task itself should be governed by this object,
+	// users should just Pin() temporarily to use it but not hold onto the reference forever.
+	TWeakPtr<FTakeUploadTask> GetUserTask() const;
+
+private:
+	// Just to be doubly safe we don't confuse the ExportTake's INDEX_NONE return for this value, we use -2
+	static constexpr int32 UnsetWorkerID = -2;
+
+	FCriticalSection CriticalSection;
+	int32 ExportManagerWorkerID = UnsetWorkerID;
+	const TSharedRef<FTakeUploadTask> SharedTask;
+	TWeakPtr<FLiveLinkHubExportManager> ExportManager;
+	const FString LogPrefix;
+};
+
+bool operator==(const FAugmentedTakeUploadTask& InLhs, const FAugmentedTakeUploadTask& InRhs)
+{
+	return InLhs.GetTaskID() == InRhs.GetTaskID();
+}
+
+FAugmentedTakeUploadTask::FAugmentedTakeUploadTask(
+	TUniquePtr<FTakeUploadTask> InTakeUploadTask,
+	TWeakPtr<FLiveLinkHubExportManager> InExportManager
+) :
+	// Convert the task into a shared ref, so we can share it in callbacks
+	SharedTask(MakeShareable<FTakeUploadTask>(InTakeUploadTask.Release())),
+	ExportManager(MoveTemp(InExportManager)),
+	LogPrefix(FString::Printf(TEXT("[Upload Task %s]"), *SharedTask->GetTaskID().ToString()))
+{
+}
+
+bool FAugmentedTakeUploadTask::Start(FLiveLinkHubExportClient::FTakeUploadParams InUploadParams)
+{
+	FScopeLock Lock(&CriticalSection);
+
+	if (ExportManagerWorkerID != UnsetWorkerID)
+	{
+		UE_LOG(
+			LogCaptureManagerUnrealEndpoint,
+			Error,
+			TEXT("%s Failed to start upload. An upload with this ID is already in progress"),
+			*LogPrefix
+		);
+		return false;
+	}
+
+	TSharedPtr<FLiveLinkHubExportManager> SharedExportManager = ExportManager.Pin();
+
+	if (SharedExportManager)
+	{
+		const int32 WorkerID = SharedExportManager->ExportTake(
+			InUploadParams,
+			SharedTask->GetDataDirectory(),
+			SharedTask->GetTakeMetadata()
+		);
+
+		if (WorkerID == INDEX_NONE)
+		{
+			UE_LOG(
+				LogCaptureManagerUnrealEndpoint,
+				Error,
+				TEXT("%s Failed to start upload. The export manager returned %d for the worker ID"),
+				*LogPrefix,
+				WorkerID
+			);
+			return false;
+		}
+
+		ExportManagerWorkerID = WorkerID;
+		return true;
+	}
+	else
+	{
+		UE_LOG(
+			LogCaptureManagerUnrealEndpoint,
+			Error,
+			TEXT("%s Failed to start upload. The export manager for this task has been destroyed"),
+			*LogPrefix
+		);
+		return false;
+	}
+}
+
+bool FAugmentedTakeUploadTask::Abort()
+{
+	TSharedPtr<FLiveLinkHubExportManager> SharedExportManager;
+	int32 WorkerID = UnsetWorkerID;
+
+	// Don't hold the lock during the ExportManager's AbortExport call, just in case we trigger some callback which
+	// tries to acquire the lock again.
+	{
+		FScopeLock Lock(&CriticalSection);
+
+		if (ExportManagerWorkerID == UnsetWorkerID)
+		{
+			// Nothing to abort - the task was likely created but no-one called Start() on it before calling Abort(),
+			// so aborting does nothing.
+			return true;
+		}
+
+		SharedExportManager = ExportManager.Pin();
+		WorkerID = ExportManagerWorkerID;
+	}
+
+	if (SharedExportManager)
+	{
+		const bool bIsAborted = SharedExportManager->AbortExport(WorkerID);
+
+		if (bIsAborted)
+		{
+			// Reset the upload ID just to avoid calling abort logic again with an ID that no longer exists
+			FScopeLock Lock(&CriticalSection);
+			ExportManagerWorkerID = UnsetWorkerID;
+			return true;
+		}
+		else
+		{
+			UE_LOG(
+				LogCaptureManagerUnrealEndpoint,
+				Error,
+				TEXT("%s Failed to abort upload. The export manager could not find a worker with this ID"),
+				*LogPrefix
+			);
+			return false;
+		}
+	}
+	else
+	{
+		UE_LOG(
+			LogCaptureManagerUnrealEndpoint,
+			Error,
+			TEXT("%s Failed to abort upload. The export manager for this take has been destroyed"),
+			*LogPrefix
+		);
+		return false;
+	}
+}
+
+FGuid FAugmentedTakeUploadTask::GetTaskID() const
+{
+	return SharedTask->GetTaskID();
+}
+
+TWeakPtr<FTakeUploadTask> FAugmentedTakeUploadTask::GetUserTask() const
+{
+	return SharedTask.ToWeakPtr();
+}
+
 struct FUnrealEndpoint::FImpl : public TSharedFromThis<FUnrealEndpoint::FImpl>
 {
-	FImpl(FUnrealEndpointInfo EndpointInfo);
+	using FMessengerType = FMessenger<FConnectStarter, FUploadStateHandler>;
+
+	explicit FImpl(FUnrealEndpointInfo EndpointInfo);
 	~FImpl();
 
-	void SetupConnectionDelegate();
+	void BuildMessenger();
 
 	void StartConnection();
 	void StopConnection();
+	void Retire();
 
 	bool WaitForConnectionState(const FUnrealEndpoint::EConnectionState InConnectionState, int32 InTimeoutMs) const;
-
-	bool AddTakeUploadTask(TUniquePtr<FTakeUploadTask> InTakeUploadTask, FGuid& OutTaskID);
+	TSharedRef<FLiveLinkHubExportManager> GetOrCreateExportManager(const FGuid& InClientID, const FGuid& InCaptureSourceID);
+	bool AddTakeUploadTask(TUniquePtr<FTakeUploadTask> InTakeUploadTask);
+	bool TaskIDNotTracked(const FGuid& TaskID);
 	void CancelTakeUploadTask(FGuid InTakeUploadTask);
-	const FUnrealEndpointInfo& GetInfo() const;
+	FUnrealEndpointInfo GetInfo() const;
 
-	FCriticalSection CriticalSection;
-	mutable std::condition_variable CondVar;
-	mutable std::mutex CondMtx;
+	void OnMessengerUploadCompleted(const FGuid& InCaptureSourceID, const FGuid& InTakeUploadID, FString InMessage, int32 InCode);
+	void OnMessengerUploadProgressed(const FGuid& InCaptureSourceID, const FGuid& InTakeUploadID, double InProgress);
+	void OnMessengerDisconnect();
+	void OnMessengerConnectResponse(const FConnectResponse& InResponse);
+
+	mutable std::condition_variable ConditionVariable;
+	mutable std::mutex Mutex;
 
 	const FUnrealEndpointInfo EndpointInfo;
+	const FString LogPrefix;
+
 	std::atomic<bool> bConnectionIsStarted;
 	bool bIsConnected;
+	std::atomic<bool> bIsRetired;
 
-	struct FUploadTrackingInfo
-	{
-		int32 UploadID = -1;
-		TWeakPtr<FUploader> Uploader;
-	};
+	// Shared ptr so we can use the presence of this ptr as a thread safe "is ready" check also
+	TSharedPtr<FMessengerType> Messenger;
 
-	TArray<TSharedRef<FTakeUploadTask>> TakeUploadTasks;
-	TMap<FGuid, FUploadTrackingInfo> UploadTrackingInfoByTaskID;
-	FCommunicationManager CommunicationManager;
+	TMap<FGuid, TSharedRef<FLiveLinkHubExportManager>> ExportManagersByCaptureSourceID;
+	TArray<TSharedRef<FAugmentedTakeUploadTask>> TakeUploadTasks;
 };
 
 FUnrealEndpoint::FUnrealEndpoint(FUnrealEndpointInfo EndpointInfo) :
 	Impl(MakeShared<FImpl>(EndpointInfo))
 {
-	Impl->SetupConnectionDelegate();
 }
 
 FUnrealEndpoint::~FUnrealEndpoint() = default;
@@ -63,14 +235,19 @@ void FUnrealEndpoint::StopConnection()
 	Impl->StopConnection();
 }
 
+void FUnrealEndpoint::Retire()
+{
+	Impl->Retire();
+}
+
 bool FUnrealEndpoint::WaitForConnectionState(const FUnrealEndpoint::EConnectionState InConnectionState, int32 InTimeoutMs) const
 {
 	return Impl->WaitForConnectionState(InConnectionState, InTimeoutMs);
 }
 
-bool FUnrealEndpoint::AddTakeUploadTask(TUniquePtr<FTakeUploadTask> InTakeUploadTask, FGuid& OutTaskID)
+bool FUnrealEndpoint::AddTakeUploadTask(TUniquePtr<FTakeUploadTask> InTakeUploadTask)
 {
-	return Impl->AddTakeUploadTask(MoveTemp(InTakeUploadTask), OutTaskID);
+	return Impl->AddTakeUploadTask(MoveTemp(InTakeUploadTask));
 }
 
 void FUnrealEndpoint::CancelTakeUploadTask(FGuid InTakeUploadTask)
@@ -78,86 +255,202 @@ void FUnrealEndpoint::CancelTakeUploadTask(FGuid InTakeUploadTask)
 	Impl->CancelTakeUploadTask(MoveTemp(InTakeUploadTask));
 }
 
-const FUnrealEndpointInfo& FUnrealEndpoint::GetInfo() const
+FUnrealEndpointInfo FUnrealEndpoint::GetInfo() const
 {
 	return Impl->GetInfo();
 }
 
 FUnrealEndpoint::FImpl::FImpl(FUnrealEndpointInfo InEndpointInfo) :
 	EndpointInfo(MoveTemp(InEndpointInfo)),
+	LogPrefix(FString::Printf(TEXT("[%s]"), *UnrealEndpointInfoToString(EndpointInfo))),
 	bConnectionIsStarted(false),
-	bIsConnected(false)
+	bIsConnected(false),
+	bIsRetired(false)
 {
 }
 
-FUnrealEndpoint::FImpl::~FImpl() = default;
+FUnrealEndpoint::FImpl::~FImpl()
+{
+	// We do not want to acquire any mutexes or do any complex work here, so the user is required to call StopConnection
+	// before the endpoint is destroyed
+	checkf(!bConnectionIsStarted, TEXT("The user is expected to have called StopConnection before destruction"));
+};
 
-void FUnrealEndpoint::FImpl::SetupConnectionDelegate()
+void FUnrealEndpoint::FImpl::BuildMessenger()
 {
-	CommunicationManager.ConnectionStateChanged().BindLambda(
-		[WeakThis = AsWeak()](const FMessageAddress& InMessageAddress, bool bInIsConnected)
-		{
-			if (TSharedPtr<FImpl> This = WeakThis.Pin())
+	TRACE_CPUPROFILER_EVENT_SCOPE(FUnrealEndpoint_BuildMessenger);
+
+	TSharedPtr<FMessengerType> TheMessenger = MakeShared<FMessengerType>();
+
+	TheMessenger->SetDisconnectHandler(FConnectStarter::FDisconnectHandler::CreateRaw(this, &FImpl::OnMessengerDisconnect));
+	TheMessenger->SetUploadCallbacks(
+		FUploadStateHandler::FUploadStateCallback::CreateLambda(
+			[WeakThis = AsWeak()](const FGuid& InCaptureSourceID, const FGuid& InTakeUploadID, double InProgress)
 			{
-				if (InMessageAddress == This->EndpointInfo.MessageAddress)
-				{
+				using namespace UE::Tasks;
+
+				auto Task = [WeakThis, InCaptureSourceID, InTakeUploadID, InProgress]()
 					{
-						std::lock_guard<std::mutex> LockGuard(This->CondMtx);
-						This->bIsConnected = bInIsConnected;
-					}
-					This->CondVar.notify_one();
-				}
+						TSharedPtr<FUnrealEndpoint::FImpl> This = WeakThis.Pin();
+
+						if (This)
+						{
+							This->OnMessengerUploadProgressed(InCaptureSourceID, InTakeUploadID, InProgress);
+						}
+					};
+
+				// Offload the task onto a background thread (i.e. never on a foreground thread, unlike AsyncTask/AnyThread).
+				Launch(UE_SOURCE_LOCATION, MoveTemp(Task), ETaskPriority::BackgroundNormal);
 			}
-		}
+		),
+		FUploadStateHandler::FUploadFinishedCallback::CreateLambda(
+			[WeakThis = AsWeak()](const FGuid& InCaptureSourceID, const FGuid& InTakeUploadID, FString InMessage, const int32 InCode)
+			{
+				using namespace UE::Tasks;
+
+				auto Task = [WeakThis, InCaptureSourceID, InTakeUploadID, Message = MoveTemp(InMessage), InCode]() mutable
+					{
+						TSharedPtr<FUnrealEndpoint::FImpl> This = WeakThis.Pin();
+
+						if (This)
+						{
+							This->OnMessengerUploadCompleted(InCaptureSourceID, InTakeUploadID, MoveTemp(Message), InCode);
+						}
+					};
+
+				// Offload the task onto a background thread (i.e. never on a foreground thread, unlike AsyncTask/AnyThread).
+				// Basically so we don't block the message bus by doing things like directory traversal/removal.
+				Launch(UE_SOURCE_LOCATION, MoveTemp(Task), ETaskPriority::BackgroundNormal);
+			}
+		)
 	);
+
+	// EndpointInfo is const member so no need to protect it with a mutex lock
+	TheMessenger->SetAddress(EndpointInfo.MessageAddress);
+	TheMessenger->Connect(FConnectStarter::FConnectHandler::CreateRaw(this, &FUnrealEndpoint::FImpl::OnMessengerConnectResponse));
+
+	// Assign the messenger member last, so we can use it as a thread-safe indicator that the messenger is ready
+	std::lock_guard<std::mutex> Lock(Mutex);
+	Messenger = MoveTemp(TheMessenger);
 }
 
 void FUnrealEndpoint::FImpl::StartConnection()
 {
+	TRACE_CPUPROFILER_EVENT_SCOPE(FUnrealEndpoint_StartConnection);
+
+	if (bIsRetired)
+	{
+		UE_LOG(
+			LogCaptureManagerUnrealEndpoint,
+			Warning,
+			TEXT("%s Ignoring request to start connection, this endpoint has been retired"),
+			*LogPrefix
+		);
+		return;
+	}
+
 	if (bConnectionIsStarted)
 	{
+		// Nothing to do
 		return;
 	}
 
 	bConnectionIsStarted = true;
-
-	FScopeLock ScopeLock(&CriticalSection);
-	CommunicationManager.Connect(EndpointInfo.MessageAddress, EndpointInfo.IPAddress, EndpointInfo.ImportServicePort);
+	BuildMessenger();
 }
 
 void FUnrealEndpoint::FImpl::StopConnection()
 {
+	TRACE_CPUPROFILER_EVENT_SCOPE(FUnrealEndpoint_StopConnection);
+
 	if (!bConnectionIsStarted)
 	{
 		// Nothing to do
+		check(!bIsConnected);
 		return;
 	}
 
-	FScopeLock ScopeLock(&CriticalSection);
-	CommunicationManager.Disconnect();
+	// We don't hold the lock while we tear down the messenger/abort the tasks, as the callbacks for disconnection etc.
+	// may try to acquire it through the public interface.
+	TSharedPtr<FMessengerType> TheMessenger;
+	TArray<TSharedRef<FAugmentedTakeUploadTask>> TasksToAbort;
+
+	{
+		std::lock_guard<std::mutex> Lock(Mutex);
+		TheMessenger = Messenger;
+		Messenger.Reset();
+
+		TasksToAbort = MoveTemp(TakeUploadTasks);
+		TakeUploadTasks.Empty();
+	}
+
+	for (TSharedRef<FAugmentedTakeUploadTask>& TaskToAbort : TasksToAbort)
+	{
+		const bool bIsAborted = TaskToAbort->Abort();
+
+		if (!bIsAborted)
+		{
+			UE_LOG(
+				LogCaptureManagerUnrealEndpoint,
+				Warning,
+				TEXT("%s Failed to abort take upload task %s while stopping the connection"),
+				*LogPrefix,
+				*TaskToAbort->GetTaskID().ToString()
+			);
+		}
+	}
+
+	if (TheMessenger)
+	{
+		TheMessenger->Disconnect();
+		TheMessenger->SetDisconnectHandler(nullptr);
+		TheMessenger->SetUploadCallbacks(nullptr, nullptr);
+
+		// Make sure the messenger's teardown is complete before we set status bools/return
+		TheMessenger.Reset();
+	}
+
+	check(!bIsConnected);
 	bConnectionIsStarted = false;
 }
 
+void FUnrealEndpoint::FImpl::Retire()
+{
+	TRACE_CPUPROFILER_EVENT_SCOPE(FUnrealEndpoint_Retire);
+
+	// This atomic flag is used to ensure the user doesn't try to do any meaningful work while the endpoint is being
+	// stopped, this can happen if the user holds the last reference after the endpoint manager has removed its own
+	// reference, the manager will call StopConnection but the user might still try to use the endpoint between then and
+	// when they finally destroy the endpoint on their side. This ensures that once the manager has called retire() on the
+	// endpoint, the user can't really do anything with it - the manager is responsible for the endpoint lifetime, not the
+	// user.
+	bIsRetired = true;
+
+	StopConnection();
+}
+
 bool FUnrealEndpoint::FImpl::WaitForConnectionState(const FUnrealEndpoint::EConnectionState InConnectionState, const int32 TimeoutMs) const
 {
-	std::unique_lock<std::mutex> Lock(CondMtx);
+	TRACE_CPUPROFILER_EVENT_SCOPE(FUnrealEndpoint_WaitForConnectionState);
 
-	bool bWaitSuccess = CondVar.wait_for(
+	std::unique_lock<std::mutex> Lock(Mutex);
+
+	bool bWaitSuccess = ConditionVariable.wait_for(
 		Lock,
 		std::chrono::milliseconds(TimeoutMs),
-		[This = AsShared(), InConnectionState]() -> bool
+		[this, InConnectionState]() -> bool
 		{
 			switch (InConnectionState)
 			{
-				case EConnectionState::Connected:
-					return This->bIsConnected;
+			case EConnectionState::Connected:
+				return bIsConnected;
 
-				case EConnectionState::Disconnected:
-					return !This->bIsConnected;
+			case EConnectionState::Disconnected:
+				return !bIsConnected;
 
-				default:
-					check(false);
-					return false;
+			default:
+				check(false);
+				return false;
 			}
 		}
 	);
@@ -165,130 +458,324 @@ bool FUnrealEndpoint::FImpl::WaitForConnectionState(const FUnrealEndpoint::EConn
 	return bWaitSuccess;
 }
 
-bool FUnrealEndpoint::FImpl::AddTakeUploadTask(TUniquePtr<FTakeUploadTask> InTakeUploadTask, FGuid& OutTaskID)
+TSharedRef<FLiveLinkHubExportManager> FUnrealEndpoint::FImpl::GetOrCreateExportManager(
+	const FGuid& InClientID,
+	const FGuid& InCaptureSourceID
+)
 {
-	if (!InTakeUploadTask)
+	// Note: Assumes mutex lock is already held by the caller
+
+	TSharedRef<FLiveLinkHubExportManager>* FoundExportManager = ExportManagersByCaptureSourceID.Find(InCaptureSourceID);
+
+	if (FoundExportManager)
 	{
-		check(false);
-		return false;
+		return *FoundExportManager;
 	}
 
-	// Convert the task into a shared ref, so we can share it in callbacks
-	TSharedRef<FTakeUploadTask> SharedTask = MakeShareable<FTakeUploadTask>(InTakeUploadTask.Release());
-
-	// Note: Only the first call to GetUploader will register these callbacks
-	FUploader::FUploadStateUpdate UpdateCallback = FUploader::FUploadStateUpdate::CreateLambda(
-		[WeakThis = AsWeak()](const FGuid& InTakeUploadID, const double InProgress)
+	auto OnDataUploaded = [WeakThis = AsWeak(), CaptureSourceID = InCaptureSourceID](const FGuid& InTakeUploadID, FUploadVoidResult InResult)
 		{
-			if (TSharedPtr<FImpl> This = WeakThis.Pin())
+			if (InResult.HasError())
 			{
-				// Find the take upload task and update it
-				FScopeLock ScopeLock(&This->CriticalSection);
+				TSharedPtr<FUnrealEndpoint::FImpl> This = WeakThis.Pin();
 
-				const int32 TaskIndex = This->TakeUploadTasks.IndexOfByPredicate(
-					[&InTakeUploadID](const TSharedRef<FTakeUploadTask>& InTask)
-					{
-						return InTask->GetTaskID() == InTakeUploadID;
-					}
-				);
-
-				if (TaskIndex != INDEX_NONE)
+				if (This)
 				{
-					// Grab a strong reference before releasing the lock
-					TSharedRef<FTakeUploadTask> ThisTask = This->TakeUploadTasks[TaskIndex];
-					ScopeLock.Unlock();
-
-					ThisTask->Progressed().ExecuteIfBound(InProgress);
+					This->OnMessengerUploadCompleted(CaptureSourceID, InTakeUploadID, InResult.GetError().GetText().ToString(), InResult.GetError().GetCode());
+				}
+				else
+				{
+					UE_LOG(
+						LogCaptureManagerUnrealEndpoint,
+						Error,
+						TEXT("Failed to report export manager error to the endpoint, the endpoint has been destroyed: ErrorMessage=%s, ErrorCode=%d"),
+						*InResult.GetError().GetText().ToString(),
+						InResult.GetError().GetCode()
+					);
 				}
 			}
-		}
+		};
+
+	return ExportManagersByCaptureSourceID.Emplace(
+		InCaptureSourceID,
+		MakeShared<FLiveLinkHubExportManager>(
+			InClientID,
+			FLiveLinkHubExportClient::FOnDataUploaded::CreateLambda(MoveTemp(OnDataUploaded))
+		)
 	);
+};
 
-	// Note: Only the first call to GetUploader will register these callbacks
-	FUploader::FUploadFinished UpdateFinishedCallback = FUploader::FUploadFinished::CreateLambda(
-		[WeakThis = AsWeak()](const FGuid& InTakeUploadID, const FString& InMessage, int32 InCode)
-		{
-			if (TSharedPtr<FImpl> This = WeakThis.Pin())
-			{
-				FScopeLock ScopeLock(&This->CriticalSection);
+bool FUnrealEndpoint::FImpl::AddTakeUploadTask(TUniquePtr<FTakeUploadTask> InTakeUploadTask)
+{
+	TRACE_CPUPROFILER_EVENT_SCOPE(FUnrealEndpoint_AddTakeUploadTask);
 
-				// Find the take upload task and update it
-				const int32 TaskIndex = This->TakeUploadTasks.IndexOfByPredicate(
-					[&InTakeUploadID](const TSharedRef<FTakeUploadTask>& InTask)
-					{
-						return InTask->GetTaskID() == InTakeUploadID;
-					}
-				);
+	if (!ensure(InTakeUploadTask))
+	{
+		UE_LOG(
+			LogCaptureManagerUnrealEndpoint,
+			Error,
+			TEXT("%s Failed to upload take, the supplied task was invalid (nullptr)"),
+			*LogPrefix
+		);
+		return false;
+	}
 
-				if (TaskIndex != INDEX_NONE)
-				{
-					// Grab a strong reference before releasing the lock
-					TSharedRef<FTakeUploadTask> ThisTask = This->TakeUploadTasks[TaskIndex];
-					This->TakeUploadTasks.RemoveAt(TaskIndex);
-					ScopeLock.Unlock();
+	if (bIsRetired)
+	{
+		UE_LOG(
+			LogCaptureManagerUnrealEndpoint,
+			Error,
+			TEXT("%s Failed to upload take, the endpoint has been retired"),
+			*LogPrefix
+		);
+		return false;
+	}
 
-					ThisTask->Complete().ExecuteIfBound(InMessage, InCode);
-				}
-			}
-		}
-	);
+	std::lock_guard<std::mutex> Lock(Mutex);
+
+	if (!ensure(Messenger))
+	{
+		UE_LOG(
+			LogCaptureManagerUnrealEndpoint,
+			Error,
+			TEXT("%s, Failed to upload take, the underlying messenger has not been initialized"),
+			*LogPrefix
+		);
+		return false;
+	}
+
+	FGuid ClientID;
+	const bool bParseOk = FGuid::Parse(Messenger->GetOwnAddress().ToString(), ClientID);
+
+	if (!ensure(bParseOk))
+	{
+		UE_LOG(
+			LogCaptureManagerUnrealEndpoint,
+			Error,
+			TEXT("%s Failed to upload take, we failed to parse the address of the messenger"),
+			*LogPrefix
+		);
+		return false;
+	}
+
+	TSharedRef<FLiveLinkHubExportManager> ExportManager = GetOrCreateExportManager(ClientID, InTakeUploadTask->GetCaptureSourceID());
+
+	const FGuid TaskID = InTakeUploadTask->GetTaskID();
+
+	FLiveLinkHubExportClient::FTakeUploadParams UploadParams;
+	UploadParams.CaptureSourceId = InTakeUploadTask->GetCaptureSourceID();
+	UploadParams.CaptureSourceName = InTakeUploadTask->GetCaptureSourceName();
+	UploadParams.TakeUploadId = TaskID;
+	UploadParams.IpAddress = EndpointInfo.IPAddress;
+	UploadParams.Port = EndpointInfo.ImportServicePort;
+
+	checkf(TaskIDNotTracked(TaskID), TEXT("Task ID is already being tracked, it must be unique for each task"));
+
+	TSharedRef<FAugmentedTakeUploadTask> TakeUploadTask = MakeShared<FAugmentedTakeUploadTask>(MoveTemp(InTakeUploadTask), ExportManager.ToWeakPtr());
+
+	const bool bUploadIsStarted = TakeUploadTask->Start(MoveTemp(UploadParams));
+
+	if (bUploadIsStarted)
+	{
+		TakeUploadTasks.Emplace(MoveTemp(TakeUploadTask));
+	}
+
+	// Failure reason should have already been logged above
+	return bUploadIsStarted;
+}
 
-	FScopeLock ScopeLock(&CriticalSection);
+bool FUnrealEndpoint::FImpl::TaskIDNotTracked(const FGuid& TaskID)
+{
+	// Note: Assumes a mutex lock is already held by the caller
 
-	TWeakPtr<FUploader> WeakUploader = CommunicationManager.GetUploader(
-		SharedTask->GetCaptureSourceID(),
-		SharedTask->GetCaptureSourceName(),
-		MoveTemp(UpdateCallback),
-		MoveTemp(UpdateFinishedCallback)
+	const TSharedRef<FAugmentedTakeUploadTask>* TaskWithThisID = TakeUploadTasks.FindByPredicate(
+		[&TaskID](const TSharedRef<FAugmentedTakeUploadTask>& InTask)
+		{
+			return InTask->GetTaskID() == TaskID;
+		}
 	);
 
-	TSharedPtr<FUploader> Uploader = WeakUploader.Pin();
+	return TaskWithThisID == nullptr;
+}
+
+void FUnrealEndpoint::FImpl::CancelTakeUploadTask(const FGuid InTakeUploadTaskID)
+{
+	TRACE_CPUPROFILER_EVENT_SCOPE(FUnrealEndpoint_CancelTakeUploadTask);
 
-	if (!Uploader)
+	if (bIsRetired)
 	{
-		const FString TakeInfoString = FString::Printf(
-			TEXT("Slate=%s, Take=%d, UniqueID=%s)"),
-			*SharedTask->GetTakeMetadata().Slate,
-			SharedTask->GetTakeMetadata().TakeNumber,
-			*SharedTask->GetTakeMetadata().UniqueId
+		UE_LOG(
+			LogCaptureManagerUnrealEndpoint,
+			Error,
+			TEXT("%s Failed to cancel task %s. This endpoint has been retired"),
+			*LogPrefix,
+			*InTakeUploadTaskID.ToString()
 		);
-		UE_LOG(LogCaptureManagerUnrealEndpoint, Error, TEXT("Failed to find an uploader for %s"), *TakeInfoString);
+		return;
+	}
 
-		return false;
+	std::lock_guard<std::mutex> Lock(Mutex);
+
+	TSharedRef<FAugmentedTakeUploadTask>* TakeUploadTask = TakeUploadTasks.FindByPredicate(
+		[&InTakeUploadTaskID](const TSharedRef<FAugmentedTakeUploadTask>& InTask)
+		{
+			return InTask->GetTaskID() == InTakeUploadTaskID;
+		}
+	);
+
+	if (TakeUploadTask)
+	{
+		const bool bUploadIsAborted = (*TakeUploadTask)->Abort();
+
+		if (!bUploadIsAborted)
+		{
+			UE_LOG(
+				LogCaptureManagerUnrealEndpoint,
+				Error,
+				TEXT("%s Failed to abort task %s"),
+				*LogPrefix,
+				*(*TakeUploadTask)->GetTaskID().ToString()
+			);
+		}
 	}
 
-	const FGuid& TaskID = SharedTask->GetTaskID();
-	const int32 UploadID = Uploader->UploadTake(TaskID, SharedTask->GetDataDirectory(), SharedTask->GetTakeMetadata());
+	// If the task doesn't exist, this indicates that it has already completed (and been removed) before or during this
+	// cancel call
+}
+
+void FUnrealEndpoint::FImpl::OnMessengerUploadCompleted(
+	const FGuid& InCaptureSourceID,
+	const FGuid& InTakeUploadID,
+	const FString InMessage,
+	const int32 InCode
+)
+{
+	TRACE_CPUPROFILER_EVENT_SCOPE(FUnrealEndpoint_OnMessengerUploadCompleted);
+
+	std::unique_lock<std::mutex> Lock(Mutex);
+
+	const TSharedRef<FAugmentedTakeUploadTask>* FoundTask = TakeUploadTasks.FindByPredicate(
+		[&InTakeUploadID](const TSharedRef<FAugmentedTakeUploadTask>& InTask)
+		{
+			return InTask->GetTaskID() == InTakeUploadID;
+		}
+	);
+
+	if (FoundTask)
+	{
+		TSharedPtr<FTakeUploadTask> TakeUploadTask = (*FoundTask)->GetUserTask().Pin();
+
+		// Don't need to track this task anymore
+		const int32 Index = static_cast<int32>(FoundTask - TakeUploadTasks.GetData());
+		TakeUploadTasks.RemoveAt(Index);
 
-	FUploadTrackingInfo UploadTrackingInfo;
-	UploadTrackingInfo.UploadID = UploadID;
-	UploadTrackingInfo.Uploader = WeakUploader;
-	UploadTrackingInfoByTaskID.Emplace(TaskID, MoveTemp(UploadTrackingInfo));
+		Lock.unlock();
 
-	TakeUploadTasks.Emplace(MoveTemp(SharedTask));
+		if (TakeUploadTask)
+		{
+			TakeUploadTask->Complete().ExecuteIfBound(InMessage, InCode);
+		}
+		else
+		{
+			UE_LOG(
+				LogCaptureManagerUnrealEndpoint,
+				Error,
+				TEXT(
+					"%s Failed to update completion state for task %s (task has already been destroyed). "
+					"Completion state was (Message=%s, Code=%d)"
+				),
+				*LogPrefix,
+				*InTakeUploadID.ToString(),
+				*InMessage,
+				InCode
+			);
+		}
+	}
 
-	OutTaskID = TaskID;
-	return true;
+	// The task may not be found and this is normal. The reason is: the OnDataUploaded handler for the export manager
+	// can race against the message bus completion message to set the task state, and the first one that invokes 
+	// OnMessengerUploadCompleted will win and remove the task, the one which comes after will be too late. This is 
+	// how the old communication manager did things so we have preserved this behavior.
 }
 
-void FUnrealEndpoint::FImpl::CancelTakeUploadTask(const FGuid InTakeUploadTaskId)
+void FUnrealEndpoint::FImpl::OnMessengerUploadProgressed(
+	const FGuid& InCaptureSourceID,
+	const FGuid& InTakeUploadID,
+	double InProgress
+)
 {
-	FScopeLock ScopeLock(&CriticalSection);
+	TRACE_CPUPROFILER_EVENT_SCOPE(FUnrealEndpoint_OnMessengerUploadProgressed);
 
-	const FUploadTrackingInfo* UploadTrackingInfo = UploadTrackingInfoByTaskID.Find(InTakeUploadTaskId);
+	std::unique_lock<std::mutex> Lock(Mutex);
 
-	if (UploadTrackingInfo)
+	const TSharedRef<FAugmentedTakeUploadTask>* TakeUploadTask = TakeUploadTasks.FindByPredicate(
+		[&InTakeUploadID](const TSharedRef<FAugmentedTakeUploadTask>& InTask)
+		{
+			return InTask->GetTaskID() == InTakeUploadID;
+		}
+	);
+
+	if (TakeUploadTask)
 	{
-		TSharedPtr<FUploader> SharedUploader = UploadTrackingInfo->Uploader.Pin();
+		TSharedPtr<FTakeUploadTask> TheTakeUploadTask = (*TakeUploadTask)->GetUserTask().Pin();
+		Lock.unlock();
 
-		if (SharedUploader)
+		if (TheTakeUploadTask)
+		{
+			TheTakeUploadTask->Progressed().ExecuteIfBound(InProgress);
+		}
+		else
 		{
-			SharedUploader->AbortUpload(UploadTrackingInfo->UploadID);
+			UE_LOG(
+				LogCaptureManagerUnrealEndpoint,
+				Error,
+				TEXT("%s Failed to update progress for task %s (task has already been destroyed). Progress was %.4f"),
+				*LogPrefix,
+				*InTakeUploadID.ToString(),
+				InProgress
+			);
 		}
 	}
+	else
+	{
+		UE_LOG(
+			LogCaptureManagerUnrealEndpoint,
+			Error,
+			TEXT("%s Failed to update progress for task %s (task not found). Progress was %.4f"),
+			*LogPrefix,
+			*InTakeUploadID.ToString(),
+			InProgress
+		);
+	}
+}
+
+void FUnrealEndpoint::FImpl::OnMessengerDisconnect()
+{
+	{
+		std::lock_guard<std::mutex> Lock(Mutex);
+		ExportManagersByCaptureSourceID.Empty();
+		bIsConnected = false;
+	}
+
+	ConditionVariable.notify_one();
+}
+
+void FUnrealEndpoint::FImpl::OnMessengerConnectResponse(const FConnectResponse& InResponse)
+{
+	TRACE_CPUPROFILER_EVENT_SCOPE(FUnrealEndpoint_OnMessengerConnectResponse);
+
+	{
+		std::lock_guard<std::mutex> Lock(Mutex);
+		bIsConnected = InResponse.Status == EStatus::Ok;
+	}
+	ConditionVariable.notify_one();
+
+	if (!bIsConnected)
+	{
+		UE_LOG(LogCaptureManagerUnrealEndpoint, Error, TEXT("%s Failed to connect to the client"), *LogPrefix);
+	}
 }
 
-const FUnrealEndpointInfo& FUnrealEndpoint::FImpl::GetInfo() const
+FUnrealEndpointInfo FUnrealEndpoint::FImpl::GetInfo() const
 {
 	return EndpointInfo;
 }
@@ -303,7 +790,13 @@ FTakeUploadTask::FUploadComplete& FTakeUploadTask::Complete()
 	return CompleteDelegate;
 }
 
-FTakeUploadTask::FTakeUploadTask(FGuid InTaskID, FGuid InCaptureSourceID, FString InCaptureSourceName, FString InDataDirectory, FTakeMetadata InTakeMetadata) :
+FTakeUploadTask::FTakeUploadTask(
+	FGuid InTaskID,
+	FGuid InCaptureSourceID,
+	FString InCaptureSourceName,
+	FString InDataDirectory,
+	FTakeMetadata InTakeMetadata
+) :
 	TaskID(MoveTemp(InTaskID)),
 	CaptureSourceID(MoveTemp(InCaptureSourceID)),
 	CaptureSourceName(MoveTemp(InCaptureSourceName)),
diff --git a/Engine/Plugins/VirtualProduction/CaptureManager/CaptureManagerApp/Source/CaptureManagerUnrealEndpoint/Private/CaptureManagerUnrealEndpointManager.cpp b/Engine/Plugins/VirtualProduction/CaptureManager/CaptureManagerApp/Source/CaptureManagerUnrealEndpoint/Private/CaptureManagerUnrealEndpointManager.cpp
--- a/Engine/Plugins/VirtualProduction/CaptureManager/CaptureManagerApp/Source/CaptureManagerUnrealEndpoint/Private/CaptureManagerUnrealEndpointManager.cpp
+++ b/Engine/Plugins/VirtualProduction/CaptureManager/CaptureManagerApp/Source/CaptureManagerUnrealEndpoint/Private/CaptureManagerUnrealEndpointManager.cpp
@@ -145,15 +145,20 @@ void FUnrealEndpointManager::FImpl::Start()
 		{
 			std::unique_lock<std::mutex> LockGuard(Mutex);
 
-			const int32 IndexToRemove = Endpoints.IndexOfByPredicate(
+			const TSharedRef<FUnrealEndpoint>* FoundEndpoint = Endpoints.FindByPredicate(
 				[&InEndpointId](const TSharedRef<FUnrealEndpoint>& InEndpoint)
 				{
 					return InEndpoint->GetInfo().EndpointID == InEndpointId;
 				}
 			);
 
-			if (IndexToRemove != INDEX_NONE)
+			if (FoundEndpoint)
 			{
+				const int32 IndexToRemove = static_cast<int32>(FoundEndpoint - Endpoints.GetData());
+
+				// Retire the endpoint and prevent others from using/restarting it. The manager is no longer tracking 
+				// this endpoint so we don't want the user to hold onto it or to start it back up.
+				(*FoundEndpoint)->Retire();
 				Endpoints.RemoveAt(IndexToRemove);
 			}
 
@@ -177,6 +182,11 @@ void FUnrealEndpointManager::FImpl::Stop()
 		DiscoveryRequester.Reset();
 	}
 
+	for (const TSharedRef<FUnrealEndpoint>& Endpoint : Endpoints)
+	{
+		Endpoint->StopConnection();
+	}
+
 	Endpoints.Empty();
 	bIsRunning = false;
 }
diff --git a/Engine/Plugins/VirtualProduction/CaptureManager/CaptureManagerApp/Source/CaptureManagerUnrealEndpoint/Private/LiveLinkHubExportClient.cpp b/Engine/Plugins/VirtualProduction/CaptureManager/CaptureManagerApp/Source/CaptureManagerUnrealEndpoint/Private/LiveLinkHubExportClient.cpp
--- a/Engine/Plugins/VirtualProduction/CaptureManager/CaptureManagerApp/Source/CaptureManagerUnrealEndpoint/Private/LiveLinkHubExportClient.cpp
+++ b/Engine/Plugins/VirtualProduction/CaptureManager/CaptureManagerApp/Source/CaptureManagerUnrealEndpoint/Private/LiveLinkHubExportClient.cpp
@@ -184,6 +184,8 @@ uint64 FLiveLinkHubExportClient::GetTotalSizeBytes(const FString& InTakeStorage)
 
 void FLiveLinkHubExportClient::OnUploadTake(TUniquePtr<FTransferContext> InTransferContext)
 {
+	TRACE_CPUPROFILER_EVENT_SCOPE(FLiveLinkHubExportClient_OnUploadTake);
+
 	if (!InTransferContext)
 	{
 		return;
@@ -192,13 +194,24 @@ void FLiveLinkHubExportClient::OnUploadTake(TUniquePtr<FTransferContext> InTrans
 	FTransferContext& Context = *InTransferContext;
 	const FTakeUploadParams& TakeUploadParams = Context.TakeUploadParams;
 
-	UE::CaptureManager::FStopToken Token = StopRequester.CreateToken();
-
 	ON_SCOPE_EXIT
 	{
 		--TaskCounter;
 	};
 
+	// Currently this will stop all tasks currently listening to the stop token, we may want to revisit this in future
+	// to allow selectively stopping tasks.
+	const bool bIsStopRequested = StopRequester.IsStopRequested();
+
+	if (bIsStopRequested)
+	{
+		OnDataUploaded.ExecuteIfBound(TakeUploadParams.TakeUploadId, MakeError(LOCTEXT("UploadTcpClient_TakeHeaderAborted", "Aborted by the user")));
+		Disconnect();
+		return;
+	}
+
+	UE::CaptureManager::FStopToken Token = StopRequester.CreateToken();
+
 	FPaths::NormalizeDirectoryName(Context.TakeStorage);
 
 	const FTakeMetadata& TakeMetadata = Context.TakeMetadata;
diff --git a/Engine/Plugins/VirtualProduction/CaptureManager/CaptureManagerApp/Source/CaptureManagerUnrealEndpoint/Private/LiveLinkHubExportManager.cpp b/Engine/Plugins/VirtualProduction/CaptureManager/CaptureManagerApp/Source/CaptureManagerUnrealEndpoint/Private/LiveLinkHubExportManager.cpp
--- a/Engine/Plugins/VirtualProduction/CaptureManager/CaptureManagerApp/Source/CaptureManagerUnrealEndpoint/Private/LiveLinkHubExportManager.cpp
+++ b/Engine/Plugins/VirtualProduction/CaptureManager/CaptureManagerApp/Source/CaptureManagerUnrealEndpoint/Private/LiveLinkHubExportManager.cpp
@@ -30,6 +30,8 @@ int32 FLiveLinkHubExportManager::ExportTake(const FLiveLinkHubExportClient::FTak
 											const FString& InTakeDirectory,
 											const FTakeMetadata& InTakeMetadata)
 {
+	TRACE_CPUPROFILER_EVENT_SCOPE(FLiveLinkHubExportManager_ExportTake);
+
 	FScopeLock Lock(&Mutex);
 
 	if (Workers.IsEmpty())
@@ -43,14 +45,19 @@ int32 FLiveLinkHubExportManager::ExportTake(const FLiveLinkHubExportClient::FTak
 	return PreferredWorkerIndex;
 }
 
-void FLiveLinkHubExportManager::AbortExport(int32 InWorkerIndex)
+bool FLiveLinkHubExportManager::AbortExport(int32 InWorkerIndex)
 {
+	TRACE_CPUPROFILER_EVENT_SCOPE(FLiveLinkHubExportManager_AbortExport);
+
 	FScopeLock Lock(&Mutex);
 
 	if (Workers.IsValidIndex(InWorkerIndex))
 	{
 		Workers[InWorkerIndex]->AbortCurrentTakeUpload();
+		return true;
 	}
+
+	return false;
 }
 
 FGuid FLiveLinkHubExportManager::GetClientId() const
diff --git a/Engine/Plugins/VirtualProduction/CaptureManager/CaptureManagerApp/Source/CaptureManagerUnrealEndpoint/Private/LiveLinkHubExportManager.h b/Engine/Plugins/VirtualProduction/CaptureManager/CaptureManagerApp/Source/CaptureManagerUnrealEndpoint/Private/LiveLinkHubExportManager.h
--- a/Engine/Plugins/VirtualProduction/CaptureManager/CaptureManagerApp/Source/CaptureManagerUnrealEndpoint/Private/LiveLinkHubExportManager.h
+++ b/Engine/Plugins/VirtualProduction/CaptureManager/CaptureManagerApp/Source/CaptureManagerUnrealEndpoint/Private/LiveLinkHubExportManager.h
@@ -18,7 +18,7 @@ public:
 								   const FString& InTakeDirectory,
 								   const FTakeMetadata& InTakeMetadata);
 
-	void AbortExport(int32 InWorkerId);
+	bool AbortExport(int32 InWorkerId);
 
 	FGuid GetClientId() const;
 
diff --git a/Engine/Plugins/VirtualProduction/CaptureManager/CaptureManagerApp/Source/CaptureManagerUnrealEndpoint/Public/CaptureManagerUnrealEndpoint.h b/Engine/Plugins/VirtualProduction/CaptureManager/CaptureManagerApp/Source/CaptureManagerUnrealEndpoint/Public/CaptureManagerUnrealEndpoint.h
--- a/Engine/Plugins/VirtualProduction/CaptureManager/CaptureManagerApp/Source/CaptureManagerUnrealEndpoint/Public/CaptureManagerUnrealEndpoint.h
+++ b/Engine/Plugins/VirtualProduction/CaptureManager/CaptureManagerApp/Source/CaptureManagerUnrealEndpoint/Public/CaptureManagerUnrealEndpoint.h
@@ -73,7 +73,7 @@ public:
 		Connected
 	};
 
-	FUnrealEndpoint(FUnrealEndpointInfo InEndpointInfo);
+	explicit FUnrealEndpoint(FUnrealEndpointInfo InEndpointInfo);
 	~FUnrealEndpoint();
 
 	/**
@@ -90,6 +90,11 @@ public:
 	*/
 	void StopConnection();
 
+	/**
+	* @brief Permanently retires the endpoint and stops the connection. This prevents further API calls from functioning.
+	*/
+	void Retire();
+
 	/**
 	* @brief Blocks until either the connection is made or the timeout is exceeded.
 	*
@@ -102,7 +107,7 @@ public:
 	*
 	* @returns true if successful, false if errors occurred.
 	*/
-	bool AddTakeUploadTask(TUniquePtr<FTakeUploadTask> InTakeUploadTask, FGuid& OutTaskID);
+	bool AddTakeUploadTask(TUniquePtr<FTakeUploadTask> InTakeUploadTask);
 
 	/**
 	* @brief Cancels the upload task with the supplied task ID.
@@ -112,7 +117,7 @@ public:
 	/**
 	* @returns Information about the unreal endpoint.
 	*/
-	const FUnrealEndpointInfo& GetInfo() const;
+	FUnrealEndpointInfo GetInfo() const;
 
 private:
 	struct FImpl;
diff --git a/Engine/Plugins/VirtualProduction/CaptureManager/CaptureManagerApp/Source/IngestLiveLinkDevice/Private/BaseIngestLiveLinkDevice.cpp b/Engine/Plugins/VirtualProduction/CaptureManager/CaptureManagerApp/Source/IngestLiveLinkDevice/Private/BaseIngestLiveLinkDevice.cpp
--- a/Engine/Plugins/VirtualProduction/CaptureManager/CaptureManagerApp/Source/IngestLiveLinkDevice/Private/BaseIngestLiveLinkDevice.cpp
+++ b/Engine/Plugins/VirtualProduction/CaptureManager/CaptureManagerApp/Source/IngestLiveLinkDevice/Private/BaseIngestLiveLinkDevice.cpp
@@ -63,6 +63,24 @@ static void TryCleanUpWorkingDirectory(const FString& InWorkingDirectory)
 	}
 }
 
+/** Convenience RAII object to clean up the working area when the scope ends - it just saves us forgetting to do it */
+class FScopedWorkingArea
+{
+public:
+	FScopedWorkingArea(FString InDirectoryPath) :
+		DirectoryPath(MoveTemp(InDirectoryPath))
+	{
+	}
+
+	~FScopedWorkingArea()
+	{
+		TryCleanUpWorkingDirectory(DirectoryPath);
+	}
+
+private:
+	const FString DirectoryPath;
+};
+
 } // namespace UE::CaptureManager::Private
 
 class UBaseIngestLiveLinkDevice::FImpl : public TSharedFromThis<UBaseIngestLiveLinkDevice::FImpl>
@@ -134,8 +152,11 @@ void UBaseIngestLiveLinkDevice::FImpl::OnDeviceRemoved(TArray<int32> TakeIdentif
 
 void UBaseIngestLiveLinkDevice::FImpl::IngestTake(const UIngestCapability_ProcessHandle* InProcessHandle, const UIngestCapability_Options* InIngestOptions, TSharedPtr<UE::CaptureManager::FTaskProgress> InTaskProgress)
 {
+	TRACE_CPUPROFILER_EVENT_SCOPE(FBaseIngestLiveLinkDevice_IngestTake);
+
 	if (!TaskWaiter.CreateTask())
 	{
+		Owner->ExecuteProcessFinishedReporter(InProcessHandle, MakeError(FIngestCapability_Error::InternalError, TEXT("Failed to create task waiter")));
 		return;
 	}
 
@@ -190,12 +211,13 @@ void UBaseIngestLiveLinkDevice::FImpl::IngestTake(const UIngestCapability_Proces
 		Converters.Remove(TakeId);
 	}
 
+	// Shared just so we can pass it into the  completion callback without invoking the destructor (can't use unique due to branch structure below)
+	TSharedRef<FScopedWorkingArea> ScopedWorkingArea = MakeShared<FScopedWorkingArea>(Params.TakeOutputDirectory);
+
 	if (Result.HasValue())
 	{
 		if (IngestOptions->UploadHostName.IsEmpty())
 		{
-			TryCleanUpWorkingDirectory(Params.TakeOutputDirectory);
-
 			FString Message = TEXT("The upload endpoint was not specified (it was an empty string), try setting a default in the hub settings");
 			Owner->ExecuteProcessFinishedReporter(InProcessHandle, MakeError(FIngestCapability_Error::UnrealEndpointNotFound, MoveTemp(Message)));
 			return;
@@ -213,9 +235,10 @@ void UBaseIngestLiveLinkDevice::FImpl::IngestTake(const UIngestCapability_Proces
 		// The endpoint was not found at all
 		if (!WeakUnrealEndpoint)
 		{
-			TryCleanUpWorkingDirectory(Params.TakeOutputDirectory);
-
-			FString Message = FString::Printf(TEXT("Upload failed because the requested endpoint was not found: %s"), *IngestOptions->UploadHostName);
+			FString Message = FString::Printf(
+				TEXT("Upload failed because the requested endpoint was not found (check the ingest server is running in UE): %s"),
+				*IngestOptions->UploadHostName
+			);
 			Owner->ExecuteProcessFinishedReporter(InProcessHandle, MakeError(FIngestCapability_Error::UnrealEndpointNotFound, MoveTemp(Message)));
 			return;
 		}
@@ -225,8 +248,6 @@ void UBaseIngestLiveLinkDevice::FImpl::IngestTake(const UIngestCapability_Proces
 		// The endpoint was found but has since expired
 		if (!UnrealEndpoint)
 		{
-			TryCleanUpWorkingDirectory(Params.TakeOutputDirectory);
-
 			FString Message = FString::Printf(TEXT("Upload failed because the requested endpoint has expired: %s"), *IngestOptions->UploadHostName);
 			Owner->ExecuteProcessFinishedReporter(InProcessHandle, MakeError(FIngestCapability_Error::UnrealEndpointNotFound, MoveTemp(Message)));
 			return;
@@ -244,8 +265,6 @@ void UBaseIngestLiveLinkDevice::FImpl::IngestTake(const UIngestCapability_Proces
 
 		if (!bIsConnected)
 		{
-			TryCleanUpWorkingDirectory(Params.TakeOutputDirectory);
-
 			FString EndpointInfoString = UnrealEndpointInfoToString(UnrealEndpoint->GetInfo());
 			FString Message = FString::Printf(TEXT("Upload failed because we timed out connecting to the endpoint: %s"), *EndpointInfoString);
 			Owner->ExecuteProcessFinishedReporter(InProcessHandle, MakeError(FIngestCapability_Error::UnrealEndpointConnectionTimedOut, MoveTemp(Message)));
@@ -270,10 +289,8 @@ void UBaseIngestLiveLinkDevice::FImpl::IngestTake(const UIngestCapability_Proces
 
 		TakeUploadTask->Complete().BindLambda(
 			// We need to keep the Impl alive until the task is complete, as we need to use it
-			[Impl = SharedThis(this), ProcessHandle = InProcessHandle, TakeOutputDirectory = Params.TakeOutputDirectory](const FString& InMessage, const int32 InStatusCode)
+			[Impl = SharedThis(this), ProcessHandle = InProcessHandle, TakeOutputDirectory = Params.TakeOutputDirectory, ScopedWorkingArea = ScopedWorkingArea](const FString& InMessage, const int32 InStatusCode)
 			{
-				TryCleanUpWorkingDirectory(TakeOutputDirectory);
-
 				FScopeLock TakeUploadsRemoveLock(&Impl->Mutex);
 				Impl->TakeUploads.Remove(ProcessHandle->GetTakeId());
 				TakeUploadsRemoveLock.Unlock();
@@ -289,8 +306,8 @@ void UBaseIngestLiveLinkDevice::FImpl::IngestTake(const UIngestCapability_Proces
 			}
 		);
 
-		FGuid TaskID;
-		const bool bAddOk = UnrealEndpoint->AddTakeUploadTask(MoveTemp(TakeUploadTask), TaskID);
+		FGuid TaskID = TakeUploadTask->GetTaskID();
+		const bool bAddOk = UnrealEndpoint->AddTakeUploadTask(MoveTemp(TakeUploadTask));
 
 		if (bAddOk)
 		{
@@ -304,13 +321,12 @@ void UBaseIngestLiveLinkDevice::FImpl::IngestTake(const UIngestCapability_Proces
 		}
 		else
 		{
-			ensureMsgf(false, TEXT("Failed to lookup take directory for take id %d"), TakeId);
+			FString Message = FString::Printf(TEXT("Failed to add upload task: %s"), *TaskID.ToString());
+			Owner->ExecuteProcessFinishedReporter(InProcessHandle, MakeError(FIngestCapability_Error::InternalError, MoveTemp(Message)));
 		}
 	}
 	else
 	{
-		TryCleanUpWorkingDirectory(Params.TakeOutputDirectory);
-
 		FCaptureDataConverterError Error = Result.StealError();
 
 		FText Message = FText::Join(FText::FromString(TEXT("\n\r")), Error.GetErrors());
@@ -320,6 +336,8 @@ void UBaseIngestLiveLinkDevice::FImpl::IngestTake(const UIngestCapability_Proces
 
 void UBaseIngestLiveLinkDevice::FImpl::CancelIngest(int32 InTakeId)
 {
+	TRACE_CPUPROFILER_EVENT_SCOPE(FBaseIngestLiveLinkDevice_CancelIngest);
+
 	using namespace UE::CaptureManager;
 
 	FScopeLock ScopeLock(&Mutex);
@@ -360,6 +378,8 @@ FCaptureDataConverterResult<void> UBaseIngestLiveLinkDevice::FImpl::RunConversio
 {
 	using namespace UE::CaptureManager;
 
+	TRACE_CPUPROFILER_EVENT_SCOPE(FBaseIngestLiveLinkDevice_RunConversion);
+
 	FTaskProgress::FTask ConvertTask = InTaskProgress->StartTask();
 
 	FCaptureConvertVideoOutputParams VideoParams;
@@ -437,4 +457,4 @@ void UBaseIngestLiveLinkDevice::CancelIngest(int32 InTakeId)
 void UBaseIngestLiveLinkDevice::CancelIngestProcess_Implementation(const UIngestCapability_ProcessHandle* InProcessHandle)
 {
 	Impl->CancelIngest(InProcessHandle->GetTakeId());
-}
\ No newline at end of file
+}
diff --git a/Engine/Plugins/VirtualProduction/CaptureManager/CaptureManagerCore/Source/CaptureUtils/Public/Async/QueueRunner.h b/Engine/Plugins/VirtualProduction/CaptureManager/CaptureManagerCore/Source/CaptureUtils/Public/Async/QueueRunner.h
--- a/Engine/Plugins/VirtualProduction/CaptureManager/CaptureManagerCore/Source/CaptureUtils/Public/Async/QueueRunner.h
+++ b/Engine/Plugins/VirtualProduction/CaptureManager/CaptureManagerCore/Source/CaptureUtils/Public/Async/QueueRunner.h
@@ -21,7 +21,7 @@ public:
 	DECLARE_DELEGATE_OneParam(FOnProcess, T InElem)
 
 	TQueueRunner(FOnProcess InProcess)
-		: bRunning(true)
+		: bRunning(false)
 		, OnProcess(MoveTemp(InProcess))
 	{
 		Thread.Reset(FRunnableThread::Create(this, TEXT("Queue Runner"), 128 * 1024, TPri_Normal, FPlatformAffinity::GetPoolThreadMask()));
@@ -29,7 +29,10 @@ public:
 
 	virtual ~TQueueRunner() override
 	{
-		Thread->Kill(true);
+		if (Thread)
+		{
+			Thread->Kill(true);
+		}
 	}
 
 	void Add(T InElement)
@@ -57,6 +60,8 @@ protected:
 
 	virtual uint32 Run() override
 	{
+		bRunning = true;
+
 		T Elem;
 	
 		while (bRunning.load())
diff --git a/Engine/Plugins/VirtualProduction/CaptureManager/CaptureManagerCore/Source/LiveLinkHubCaptureMessaging/Private/Features/ConnectAcceptor.cpp b/Engine/Plugins/VirtualProduction/CaptureManager/CaptureManagerCore/Source/LiveLinkHubCaptureMessaging/Private/Features/ConnectAcceptor.cpp
--- a/Engine/Plugins/VirtualProduction/CaptureManager/CaptureManagerCore/Source/LiveLinkHubCaptureMessaging/Private/Features/ConnectAcceptor.cpp
+++ b/Engine/Plugins/VirtualProduction/CaptureManager/CaptureManagerCore/Source/LiveLinkHubCaptureMessaging/Private/Features/ConnectAcceptor.cpp
@@ -41,7 +41,7 @@ void FConnectAcceptor::Disconnect()
 
 		LastPingRequest = 0.0f;
 
-		ConnectionLostHandler.ExecuteIfBound(Address);
+		ConnectionLostHandler.ExecuteIfBound(GetAddress());
 	}
 }
 
@@ -52,7 +52,9 @@ bool FConnectAcceptor::IsConnected() const
 
 void FConnectAcceptor::Initialize(FMessageEndpointBuilder& InBuilder)
 {
-	InBuilder.Handling<FConnectRequest>(this, &FConnectAcceptor::ConnectRequestHandler)
+	InBuilder
+		.Handling<FConnectRequest>(this, &FConnectAcceptor::ConnectRequestHandler)
+		.Handling<FCaptureManagerHangUp>(this, &FConnectAcceptor::OnHangUpReceived)
 		.Handling<FPingMessage>(this, &FConnectAcceptor::HandleKeepAlive);
 }
 
@@ -85,7 +87,7 @@ void FConnectAcceptor::ConnectRequestHandler(const FConnectRequest& InRequest, c
 	Response->Guid = FGuid::NewGuid();
 	Response->RequestGuid = InRequest.Guid;
 
-	Endpoint->Send(Response, Address);
+	SendMessage(Response);
 
 	if (Response->Status == EStatus::Ok)
 	{
@@ -93,6 +95,16 @@ void FConnectAcceptor::ConnectRequestHandler(const FConnectRequest& InRequest, c
 	}
 }
 
+void FConnectAcceptor::OnHangUpReceived(
+	const FCaptureManagerHangUp& InRequest,
+	const TSharedRef<IMessageContext, ESPMode::ThreadSafe>& InContext
+)
+{
+	// The connection we currently have has been dropped from the other side.
+	// Disconnect so we're ready to accept new connections.
+	Disconnect();
+}
+
 void FConnectAcceptor::HandleKeepAlive(const FPingMessage& InRequest, const TSharedRef<IMessageContext, ESPMode::ThreadSafe>& InContext)
 {
 	if (!bConnected)
@@ -105,7 +117,7 @@ void FConnectAcceptor::HandleKeepAlive(const FPingMessage& InRequest, const TSha
 	Response->Guid = FGuid::NewGuid();
 	Response->RequestGuid = InRequest.Guid;
 
-	Endpoint->Send(Response, Address);
+	SendMessage(Response);
 
 	LastPingRequest = FPlatformTime::Seconds();
 }
diff --git a/Engine/Plugins/VirtualProduction/CaptureManager/CaptureManagerCore/Source/LiveLinkHubCaptureMessaging/Private/Features/ConnectStarter.cpp b/Engine/Plugins/VirtualProduction/CaptureManager/CaptureManagerCore/Source/LiveLinkHubCaptureMessaging/Private/Features/ConnectStarter.cpp
--- a/Engine/Plugins/VirtualProduction/CaptureManager/CaptureManagerCore/Source/LiveLinkHubCaptureMessaging/Private/Features/ConnectStarter.cpp
+++ b/Engine/Plugins/VirtualProduction/CaptureManager/CaptureManagerCore/Source/LiveLinkHubCaptureMessaging/Private/Features/ConnectStarter.cpp
@@ -53,24 +53,33 @@ void FConnectStarter::Connect(FConnectHandler InConnectHandler)
 
 	AddContext(Request->Guid, MoveTemp(InConnectHandler));
 
-	Endpoint->Send(Request, Address);
+	SendMessage(Request);
 }
 
 void FConnectStarter::Disconnect()
 {
+	using namespace UE::CaptureManager;
+
 	bConnected = false;
 
-	{
-		FScopeLock Lock(&Mutex);
-		Contexts.Empty();
-	}
+	FScopeLock Lock(&Mutex);
+
+	// Take a copy, so we can execute the handler without holding the lock
+	FDisconnectHandler LocalDisconnectHandler = Handler;
+	Contexts.Empty();
+	Lock.Unlock();
+
+	// Send a hangup message so that UE can drop its side of the connection
+	FCaptureManagerHangUp* HangUp = FMessageEndpoint::MakeMessage<FCaptureManagerHangUp>();
+	SendMessage(HangUp);
 
 	TimerManager->RemoveTimer(KeepAlive);
-	Handler.ExecuteIfBound();
+	LocalDisconnectHandler.ExecuteIfBound();
 }
 
 void FConnectStarter::SetDisconnectHandler(FDisconnectHandler InHandler)
 {
+	FScopeLock ScopeLock(&Mutex);
 	Handler = MoveTemp(InHandler);
 }
 
@@ -102,7 +111,7 @@ FGuid FConnectStarter::SendPingRequest(FPingCallback InResponseCb)
 		PingContexts.Add(MoveTemp(RequestGuid), MoveTemp(InResponseCb));
 	}
 
-	Endpoint->Send(Request, Address);
+	SendMessage(Request);
 
 	return RequestGuid;
 }
@@ -172,12 +181,12 @@ void FConnectStarter::OnKeepAliveInterval()
 		{
 			This->bConnected = false;
 
-			{
-				FScopeLock Lock(&This->Mutex);
-				This->Contexts.Remove(RequestGuid);
-			}
+			FScopeLock Lock(&This->Mutex);
+			This->Contexts.Remove(RequestGuid);
+			FDisconnectHandler LocalDisconnectHandler = This->Handler;
+			Lock.Unlock();
 
-			This->Handler.ExecuteIfBound();
+			LocalDisconnectHandler.ExecuteIfBound();
 		}
 
 		This->TimerManager->RemoveTimer(This->KeepAlive);
diff --git a/Engine/Plugins/VirtualProduction/CaptureManager/CaptureManagerCore/Source/LiveLinkHubCaptureMessaging/Private/Features/UploadStateHandler.cpp b/Engine/Plugins/VirtualProduction/CaptureManager/CaptureManagerCore/Source/LiveLinkHubCaptureMessaging/Private/Features/UploadStateHandler.cpp
--- a/Engine/Plugins/VirtualProduction/CaptureManager/CaptureManagerCore/Source/LiveLinkHubCaptureMessaging/Private/Features/UploadStateHandler.cpp
+++ b/Engine/Plugins/VirtualProduction/CaptureManager/CaptureManagerCore/Source/LiveLinkHubCaptureMessaging/Private/Features/UploadStateHandler.cpp
@@ -7,6 +7,7 @@ FUploadStateHandler::FUploadStateHandler() = default;
 void FUploadStateHandler::SetUploadCallbacks(FUploadStateCallback InStateCallback,
 											 FUploadFinishedCallback InFinishedCallback)
 {
+	FScopeLock ScopeLock(&CriticalSection);
 	StateCallback = MoveTemp(InStateCallback);
 	FinishedCallback = MoveTemp(InFinishedCallback);
 }
@@ -20,10 +21,20 @@ void FUploadStateHandler::Initialize(FMessageEndpointBuilder& InBuilder)
 
 void FUploadStateHandler::HandleUploadStateMessage(const FUploadState& InMessage, const TSharedRef<IMessageContext, ESPMode::ThreadSafe>& InContext)
 {
-	StateCallback.ExecuteIfBound(InMessage.CaptureSourceId, InMessage.TakeUploadId, InMessage.Progress);
+	FScopeLock ScopeLock(&CriticalSection);
+	// Take a copy, so we can execute the callback without holding the lock
+	FUploadStateCallback LocalCallback = StateCallback;
+	ScopeLock.Unlock();
+
+	LocalCallback.ExecuteIfBound(InMessage.CaptureSourceId, InMessage.TakeUploadId, InMessage.Progress);
 }
 
 void FUploadStateHandler::HandleUploadFinishedMessage(const FUploadFinished& InMessage, const TSharedRef<IMessageContext, ESPMode::ThreadSafe>& InContext)
 {
-	FinishedCallback.ExecuteIfBound(InMessage.CaptureSourceId, InMessage.TakeUploadId, InMessage.Message, static_cast<int32>(InMessage.Status));
-}
\ No newline at end of file
+	FScopeLock ScopeLock(&CriticalSection);
+	// Take a copy, so we can execute the callback without holding the lock
+	FUploadFinishedCallback LocalCallback = FinishedCallback;
+	ScopeLock.Unlock();
+
+	LocalCallback.ExecuteIfBound(InMessage.CaptureSourceId, InMessage.TakeUploadId, InMessage.Message, static_cast<int32>(InMessage.Status));
+}
diff --git a/Engine/Plugins/VirtualProduction/CaptureManager/CaptureManagerCore/Source/LiveLinkHubCaptureMessaging/Private/Features/UploadStateSender.cpp b/Engine/Plugins/VirtualProduction/CaptureManager/CaptureManagerCore/Source/LiveLinkHubCaptureMessaging/Private/Features/UploadStateSender.cpp
--- a/Engine/Plugins/VirtualProduction/CaptureManager/CaptureManagerCore/Source/LiveLinkHubCaptureMessaging/Private/Features/UploadStateSender.cpp
+++ b/Engine/Plugins/VirtualProduction/CaptureManager/CaptureManagerCore/Source/LiveLinkHubCaptureMessaging/Private/Features/UploadStateSender.cpp
@@ -12,7 +12,7 @@ void FUploadStateSender::SendUploadStateMessage(const FGuid& InCaptureSourceId,
 	Message->TakeUploadId = InTakeUploadId;
 	Message->Progress = InProgress;
 
-	Endpoint->Send(Message, Address);
+	SendMessage(Message);
 }
 
 void FUploadStateSender::SendUploadDoneMessage(const FGuid& InCaptureSourceId, const FGuid& InTakeUploadId, FString InMessage, int32 InCode)
@@ -24,12 +24,11 @@ void FUploadStateSender::SendUploadDoneMessage(const FGuid& InCaptureSourceId, c
 	Message->Message = MoveTemp(InMessage);
 	Message->Status = ConvertStatus(InCode);
 
-	Endpoint->Send(Message, Address);
+	SendMessage(Message);
 }
 
 void FUploadStateSender::Initialize(FMessageEndpointBuilder& InBuilder)
 {
-	// Empty
 }
 
 EStatus FUploadStateSender::ConvertStatus(const int32 InCode)
@@ -41,4 +40,4 @@ EStatus FUploadStateSender::ConvertStatus(const int32 InCode)
 		default:
 			return EStatus::InternalError;
 	}
-}
\ No newline at end of file
+}
diff --git a/Engine/Plugins/VirtualProduction/CaptureManager/CaptureManagerCore/Source/LiveLinkHubCaptureMessaging/Private/Messenger.cpp b/Engine/Plugins/VirtualProduction/CaptureManager/CaptureManagerCore/Source/LiveLinkHubCaptureMessaging/Private/Messenger.cpp
--- a/Engine/Plugins/VirtualProduction/CaptureManager/CaptureManagerCore/Source/LiveLinkHubCaptureMessaging/Private/Messenger.cpp
+++ b/Engine/Plugins/VirtualProduction/CaptureManager/CaptureManagerCore/Source/LiveLinkHubCaptureMessaging/Private/Messenger.cpp
@@ -4,10 +4,24 @@
 
 void FFeatureBase::SetEndpoint(TSharedPtr<FMessageEndpoint> InEndpoint)
 {
+	FScopeLock ScopeLock(&CriticalSection);
+	PRAGMA_DISABLE_DEPRECATION_WARNINGS
 	Endpoint = InEndpoint;
+	PRAGMA_ENABLE_DEPRECATION_WARNINGS
 }
 
 void FFeatureBase::SetAddress(const FMessageAddress& InAddress)
 {
+	FScopeLock ScopeLock(&CriticalSection);
+	PRAGMA_DISABLE_DEPRECATION_WARNINGS
 	Address = InAddress;
-}
\ No newline at end of file
+	PRAGMA_ENABLE_DEPRECATION_WARNINGS
+}
+
+FMessageAddress FFeatureBase::GetAddress() const
+{
+	FScopeLock ScopeLock(&CriticalSection);
+	PRAGMA_DISABLE_DEPRECATION_WARNINGS
+	return Address;
+	PRAGMA_ENABLE_DEPRECATION_WARNINGS
+}
diff --git a/Engine/Plugins/VirtualProduction/CaptureManager/CaptureManagerCore/Source/LiveLinkHubCaptureMessaging/Public/Features/ConnectAcceptor.h b/Engine/Plugins/VirtualProduction/CaptureManager/CaptureManagerCore/Source/LiveLinkHubCaptureMessaging/Public/Features/ConnectAcceptor.h
--- a/Engine/Plugins/VirtualProduction/CaptureManager/CaptureManagerCore/Source/LiveLinkHubCaptureMessaging/Public/Features/ConnectAcceptor.h
+++ b/Engine/Plugins/VirtualProduction/CaptureManager/CaptureManagerCore/Source/LiveLinkHubCaptureMessaging/Public/Features/ConnectAcceptor.h
@@ -35,6 +35,7 @@ protected:
 private:
 
 	UE_API void ConnectRequestHandler(const FConnectRequest& InRequest, const TSharedRef<IMessageContext, ESPMode::ThreadSafe>& InContext);
+	UE_API void OnHangUpReceived(const FCaptureManagerHangUp& InRequest, const TSharedRef<IMessageContext, ESPMode::ThreadSafe>& InContext);
 	UE_API void HandleKeepAlive(const FPingMessage& InRequest, const TSharedRef<IMessageContext, ESPMode::ThreadSafe>& InContext);
 
 	UE_API void CheckConnectionActivity();
@@ -42,16 +43,14 @@ private:
 	double LastPingRequest = 0.0f;
 	double InactivityTimeout = 20.0f; // seconds
 
-	FConnectAccepted ConnectHandler;
-	FConnectionLostHandler ConnectionLostHandler;
-
-	using FFeatureBase::Endpoint;
-	using FFeatureBase::Address;
-
 	TSharedRef<UE::CaptureManager::FCaptureTimerManager> TimerManager;
 	UE::CaptureManager::FCaptureTimerManager::FTimerHandle ClientActivity;
 
 	std::atomic_bool bConnected = false;
+
+	FCriticalSection CriticalSection;
+	FConnectAccepted ConnectHandler;
+	FConnectionLostHandler ConnectionLostHandler;
 };
 
 #undef UE_API
diff --git a/Engine/Plugins/VirtualProduction/CaptureManager/CaptureManagerCore/Source/LiveLinkHubCaptureMessaging/Public/Features/ConnectStarter.h b/Engine/Plugins/VirtualProduction/CaptureManager/CaptureManagerCore/Source/LiveLinkHubCaptureMessaging/Public/Features/ConnectStarter.h
--- a/Engine/Plugins/VirtualProduction/CaptureManager/CaptureManagerCore/Source/LiveLinkHubCaptureMessaging/Public/Features/ConnectStarter.h
+++ b/Engine/Plugins/VirtualProduction/CaptureManager/CaptureManagerCore/Source/LiveLinkHubCaptureMessaging/Public/Features/ConnectStarter.h
@@ -49,9 +49,6 @@ private:
 
 	UE_API void OnKeepAliveInterval();
 
-	using FFeatureBase::Endpoint;
-	using FFeatureBase::Address;
-
 	FCriticalSection Mutex;
 	TMap<FGuid, FConnectHandler> Contexts;
 
diff --git a/Engine/Plugins/VirtualProduction/CaptureManager/CaptureManagerCore/Source/LiveLinkHubCaptureMessaging/Public/Features/UploadStateHandler.h b/Engine/Plugins/VirtualProduction/CaptureManager/CaptureManagerCore/Source/LiveLinkHubCaptureMessaging/Public/Features/UploadStateHandler.h
--- a/Engine/Plugins/VirtualProduction/CaptureManager/CaptureManagerCore/Source/LiveLinkHubCaptureMessaging/Public/Features/UploadStateHandler.h
+++ b/Engine/Plugins/VirtualProduction/CaptureManager/CaptureManagerCore/Source/LiveLinkHubCaptureMessaging/Public/Features/UploadStateHandler.h
@@ -28,13 +28,10 @@ protected:
 	UE_API void Initialize(FMessageEndpointBuilder& InBuilder);
 
 private:
-
-	using FFeatureBase::Endpoint;
-	using FFeatureBase::Address;
-
 	UE_API void HandleUploadStateMessage(const FUploadState& InMessage, const TSharedRef<IMessageContext, ESPMode::ThreadSafe>& InContext);
 	UE_API void HandleUploadFinishedMessage(const FUploadFinished& InMessage, const TSharedRef<IMessageContext, ESPMode::ThreadSafe>& InContext);
 
+	FCriticalSection CriticalSection;
 	FUploadStateCallback StateCallback;
 	FUploadFinishedCallback FinishedCallback;
 };
diff --git a/Engine/Plugins/VirtualProduction/CaptureManager/CaptureManagerCore/Source/LiveLinkHubCaptureMessaging/Public/Features/UploadStateSender.h b/Engine/Plugins/VirtualProduction/CaptureManager/CaptureManagerCore/Source/LiveLinkHubCaptureMessaging/Public/Features/UploadStateSender.h
--- a/Engine/Plugins/VirtualProduction/CaptureManager/CaptureManagerCore/Source/LiveLinkHubCaptureMessaging/Public/Features/UploadStateSender.h
+++ b/Engine/Plugins/VirtualProduction/CaptureManager/CaptureManagerCore/Source/LiveLinkHubCaptureMessaging/Public/Features/UploadStateSender.h
@@ -27,9 +27,6 @@ protected:
 private:
 
 	UE_API EStatus ConvertStatus(int32 InCode);
-
-	using FFeatureBase::Endpoint;
-	using FFeatureBase::Address;
 };
 
 #undef UE_API
diff --git a/Engine/Plugins/VirtualProduction/CaptureManager/CaptureManagerCore/Source/LiveLinkHubCaptureMessaging/Public/LiveLinkHubCaptureMessages.h b/Engine/Plugins/VirtualProduction/CaptureManager/CaptureManagerCore/Source/LiveLinkHubCaptureMessaging/Public/LiveLinkHubCaptureMessages.h
--- a/Engine/Plugins/VirtualProduction/CaptureManager/CaptureManagerCore/Source/LiveLinkHubCaptureMessaging/Public/LiveLinkHubCaptureMessages.h
+++ b/Engine/Plugins/VirtualProduction/CaptureManager/CaptureManagerCore/Source/LiveLinkHubCaptureMessaging/Public/LiveLinkHubCaptureMessages.h
@@ -63,6 +63,12 @@ struct FConnectResponse : public FBaseResponse
 	GENERATED_BODY()
 };
 
+USTRUCT()
+struct FCaptureManagerHangUp : public FBaseMessage
+{
+	GENERATED_BODY()
+};
+
 USTRUCT()
 struct FDiscoveryRequest
 {
diff --git a/Engine/Plugins/VirtualProduction/CaptureManager/CaptureManagerCore/Source/LiveLinkHubCaptureMessaging/Public/Messenger.h b/Engine/Plugins/VirtualProduction/CaptureManager/CaptureManagerCore/Source/LiveLinkHubCaptureMessaging/Public/Messenger.h
--- a/Engine/Plugins/VirtualProduction/CaptureManager/CaptureManagerCore/Source/LiveLinkHubCaptureMessaging/Public/Messenger.h
+++ b/Engine/Plugins/VirtualProduction/CaptureManager/CaptureManagerCore/Source/LiveLinkHubCaptureMessaging/Public/Messenger.h
@@ -18,17 +18,60 @@ class FFeatureBase
 {
 protected:
 
+	PRAGMA_DISABLE_DEPRECATION_WARNINGS
 	virtual ~FFeatureBase() = default;
+	PRAGMA_ENABLE_DEPRECATION_WARNINGS
 
 	virtual void Initialize(FMessageEndpointBuilder& InBuilder) = 0;
 
+	template <typename MessageType>
+	void SendMessage(MessageType* InMessage);
+
 	UE_API void SetEndpoint(TSharedPtr<FMessageEndpoint> InEndpoint);
 	UE_API void SetAddress(const FMessageAddress& InAddress);
+	UE_API FMessageAddress GetAddress() const;
 
+	// When removing this deprecation, make the member private and remove warning suppression in this class
+	UE_DEPRECATED(5.7, "Accessing the endpoint directly is deprecated and likely to result in a crash, use SendMessage instead. Access to this member will be removed in a future version")
 	TSharedPtr<FMessageEndpoint> Endpoint;
+
+	// When removing this deprecation, make the member private and remove warning suppression in this classs
+	UE_DEPRECATED(5.7, "Accessing the address directly is deprecated and likely to result in a crash, use Set/GetAddress instead. Access to this member will be removed in a future version")
 	FMessageAddress Address;
+
+private:
+	mutable FCriticalSection CriticalSection;
 };
 
+template <typename MessageType>
+void FFeatureBase::SendMessage(MessageType* InMessage)
+{
+	if (!ensureMsgf(InMessage, TEXT("SendMessage() requires a valid message (it was nullptr)")))
+	{
+		return;
+	}
+
+	FScopeLock ScopeLock(&CriticalSection);
+	PRAGMA_DISABLE_DEPRECATION_WARNINGS
+	TSharedPtr<FMessageEndpoint> MessageEndpoint = Endpoint;
+	FMessageAddress MessageAddress = Address;
+	PRAGMA_ENABLE_DEPRECATION_WARNINGS
+
+	// We don't hold the lock during Send(), just in case the caller has bound a message handler which calls back into
+	// the public API of this class.
+	ScopeLock.Unlock();
+
+	if (ensureMsgf(MessageEndpoint, TEXT("A valid endpoint must be set before calling SendMessage")))
+	{
+		MessageEndpoint->Send(InMessage, MessageAddress);
+	}
+	else
+	{
+		// Since we didn't call Send, we have to free the memory for the message ourselves
+		FMemory::Free(InMessage);
+	}
+}
+
 template<class ... Features>
 class FMessenger
 	: public Features...
