Assumption that ensure will most likely never happen was wrong which is why it is promoted to user error.

diff --git a/Engine/Plugins/VirtualProduction/CaptureManager/CaptureManagerApp/Source/IngestLiveLinkDevice/Private/BaseIngestLiveLinkDevice.cpp b/Engine/Plugins/VirtualProduction/CaptureManager/CaptureManagerApp/Source/IngestLiveLinkDevice/Private/BaseIngestLiveLinkDevice.cpp
--- a/Engine/Plugins/VirtualProduction/CaptureManager/CaptureManagerApp/Source/IngestLiveLinkDevice/Private/BaseIngestLiveLinkDevice.cpp
+++ b/Engine/Plugins/VirtualProduction/CaptureManager/CaptureManagerApp/Source/IngestLiveLinkDevice/Private/BaseIngestLiveLinkDevice.cpp
@@ -185,7 +185,7 @@ void UBaseIngestLiveLinkDevice::FImpl::IngestTake(const UIngestCapability_Proces
 	if (!MaybeTakeMetadata.IsSet())
 	{
 		FString Message = FString::Printf(TEXT("Failed to lookup take metadata for take id %d (TakePath=%s)"), TakeId, *TakeFullPath);
-		ensureMsgf(false, TEXT("%s"), *Message);
+		UE_LOG(LogIngestLiveLinkDevice, Error, TEXT("%s"), *Message);
 		Owner->ExecuteProcessFinishedReporter(InProcessHandle, MakeError(FIngestCapability_Error::InvalidArgument, MoveTemp(Message)));
 		return;
 	}
