Fixing a deadlock when disconnect happens

diff --git a/Engine/Plugins/VirtualProduction/CaptureManager/CaptureManagerApp/Source/CaptureManagerUnrealEndpoint/Private/CaptureManagerUnrealEndpoint.cpp b/Engine/Plugins/VirtualProduction/CaptureManager/CaptureManagerApp/Source/CaptureManagerUnrealEndpoint/Private/CaptureManagerUnrealEndpoint.cpp
--- a/Engine/Plugins/VirtualProduction/CaptureManager/CaptureManagerApp/Source/CaptureManagerUnrealEndpoint/Private/CaptureManagerUnrealEndpoint.cpp
+++ b/Engine/Plugins/VirtualProduction/CaptureManager/CaptureManagerApp/Source/CaptureManagerUnrealEndpoint/Private/CaptureManagerUnrealEndpoint.cpp
@@ -750,12 +750,18 @@ void FUnrealEndpoint::FImpl::OnMessengerUploadProgressed(
 
 void FUnrealEndpoint::FImpl::OnMessengerDisconnect()
 {
+	TMap<FGuid, TSharedRef<FLiveLinkHubExportManager>> LocalExportManagersByCaptureSourceID;
+
 	{
 		std::lock_guard<std::mutex> Lock(Mutex);
+		LocalExportManagersByCaptureSourceID = ExportManagersByCaptureSourceID;
 		ExportManagersByCaptureSourceID.Empty();
+
 		bIsConnected = false;
 	}
 
+	LocalExportManagersByCaptureSourceID.Empty();
+
 	ConditionVariable.notify_one();
 }
 
