[CaptureManager] Fixed issue during timecode and frame rate extraction which could result in invalid values being set on the ingested data

diff --git a/Engine/Plugins/VirtualProduction/CaptureManager/CaptureManagerApp/Source/IngestLiveLinkDevice/Private/Utils/CaptureExtractTimecode.cpp b/Engine/Plugins/VirtualProduction/CaptureManager/CaptureManagerApp/Source/IngestLiveLinkDevice/Private/Utils/CaptureExtractTimecode.cpp
--- a/Engine/Plugins/VirtualProduction/CaptureManager/CaptureManagerApp/Source/IngestLiveLinkDevice/Private/Utils/CaptureExtractTimecode.cpp
+++ b/Engine/Plugins/VirtualProduction/CaptureManager/CaptureManagerApp/Source/IngestLiveLinkDevice/Private/Utils/CaptureExtractTimecode.cpp
@@ -68,7 +68,7 @@ FCaptureExtractVideoInfo::FExtractResult FCaptureExtractVideoInfo::ExtractInfo(c
 		return MakeError(ECaptureExtractInfoError::InternalError);
 	}
 
-	FCaptureExtractVideoInfo::FExtractResult Result = ExtractInfoUsingElectraPlayer(InFilePath);;
+	FCaptureExtractVideoInfo::FExtractResult Result = ExtractInfoUsingElectraPlayer(InFilePath);
 
 	if (Result.HasError())
 	{
@@ -174,6 +174,9 @@ FCaptureExtractVideoInfo::FExtractResult FCaptureExtractVideoInfo::ExtractInfoUs
 
 	bool bIsTimeout = true;
 
+	bool bFoundTimecode = false;
+	bool bFoundFrameRate = false;
+
 	while ((FDateTime::Now() - EventWaitStart).GetSeconds() < TimeoutPeriod)
 	{
 		MediaPlayer->TickInput(0, 0);
@@ -209,10 +212,7 @@ FCaptureExtractVideoInfo::FExtractResult FCaptureExtractVideoInfo::ExtractInfoUs
 				MediaPlayer->GetTracks().GetVideoTrackFormat(Track, Format, FormatInfo);
 
 				FrameRate = ConvertFrameRate(FormatInfo.FrameRate);
-			}
-			else
-			{
-				UE_LOG(LogCaptureExtractTimecode, Warning, TEXT("Failed to obtain the frame rate"));
+				bFoundFrameRate = true;
 			}
 
 			break;
@@ -241,30 +241,59 @@ FCaptureExtractVideoInfo::FExtractResult FCaptureExtractVideoInfo::ExtractInfoUs
 		const TOptional<FTimecode> Timecode = FTimecode::ParseTimecode(*TimecodeString);
 		if (!Timecode.IsSet())
 		{
-			UE_LOG(LogCaptureExtractTimecode, Warning, TEXT("Failed to parse the timecode"));
+			UE_LOG(LogCaptureExtractTimecode, Warning, TEXT("Failed to parse the timecode from %s"), *InFilePath);
 		}
 		else
 		{
 			TimecodeInfo.Timecode = Timecode.GetValue();
+			bFoundTimecode = true;
 		}
 	}
 	else
 	{
-		UE_LOG(LogCaptureExtractTimecode, Warning, TEXT("Timecode has not been found"));
+		UE_LOG(LogCaptureExtractTimecode, Warning, TEXT("Timecode has not been found in %s"), *InFilePath);
 	}
 
+	bool bFoundTimecodeRate = false;
 	if (!TimecodeRateString.IsEmpty())
 	{
 		TValueOrError<FFrameRate, UE::CaptureManager::ECaptureExtractInfoError> Result = ParseTimecodeRate(TimecodeRateString);
 		if (Result.IsValid())
 		{
 			TimecodeInfo.TimecodeRate = Result.GetValue();
+			bFoundTimecodeRate = true;
 		}
 	}
 	else
 	{
-		UE_LOG(LogCaptureExtractTimecode, Warning, TEXT("Timecode rate has not been found. Using video Frame rate as timecode rate"));
-		TimecodeInfo.TimecodeRate = FrameRate;
+		if (bFoundFrameRate)
+		{
+			// We didn't extract timecode rate, but we did find the frame rate from the video, so use that instead
+			UE_LOG(LogCaptureExtractTimecode, Warning, TEXT("Timecode rate has not been found for %s. Using video frame rate '%.2f' as timecode rate"), *InFilePath, FrameRate.AsDecimal());
+			TimecodeInfo.TimecodeRate = FrameRate;
+		}
+		else
+		{
+			UE_LOG(LogCaptureExtractTimecode, Warning, TEXT("Timecode rate has not been found for %s. Using default frame rate '%.2f' as timecode rate"), *InFilePath, FFrameRate().AsDecimal());
+		}
+	}
+
+	if (!bFoundTimecodeRate && !bFoundFrameRate)
+	{
+		return MakeError(ECaptureExtractInfoError::TimecodeRateNotFound);
+	}
+
+	if (!bFoundTimecode)
+	{
+		if (!bFoundFrameRate)
+		{
+			// Both not found - InternalError as an new enum could not be added in a hotfix release
+			return MakeError(ECaptureExtractInfoError::InternalError);
+		}
+		else
+		{
+			return MakeError(ECaptureExtractInfoError::TimecodeNotFound);
+		}
 	}
 
 	return MakeValue();
@@ -402,7 +431,8 @@ FCaptureExtractVideoInfo::FExtractResult FCaptureExtractVideoInfo::ExtractInfoUs
 
 	if (bFoundFrameRate)
 	{
-		UE_LOG(LogCaptureExtractTimecode, Warning, TEXT("Timecode rate has not been found. Using video frame rate as timecode rate"));
+		// We didn't extract timecode rate, but we did find the frame rate from the video, so use that instead
+		UE_LOG(LogCaptureExtractTimecode, Warning, TEXT("Timecode rate has not been found. Using video frame rate '%.2f' as timecode rate"), FrameRate.AsDecimal());
 		TimecodeInfo.TimecodeRate = FrameRate;
 	}
 
