[CaptureManager] No longer fall through to other timecode extraction methods if third party encoder enabled. Improved third party encoder checks and warning messages.

#Synced-CL 48154294

diff --git a/Engine/Plugins/VirtualProduction/CaptureManager/CaptureManagerApp/Source/IngestLiveLinkDevice/Private/Utils/CaptureExtractTimecode.cpp b/Engine/Plugins/VirtualProduction/CaptureManager/CaptureManagerApp/Source/IngestLiveLinkDevice/Private/Utils/CaptureExtractTimecode.cpp
--- a/Engine/Plugins/VirtualProduction/CaptureManager/CaptureManagerApp/Source/IngestLiveLinkDevice/Private/Utils/CaptureExtractTimecode.cpp
+++ b/Engine/Plugins/VirtualProduction/CaptureManager/CaptureManagerApp/Source/IngestLiveLinkDevice/Private/Utils/CaptureExtractTimecode.cpp
@@ -29,11 +29,15 @@ FCaptureExtractVideoInfo::FResult FCaptureExtractVideoInfo::Create(const FString
 	FCaptureExtractVideoInfo VideoInfoExtractor;
 
 	FExtractResult Result = VideoInfoExtractor.ExtractInfo(InFilePath);
+
 	if (Result.HasError())
 	{
+		UE_LOG(LogCaptureExtractTimecode, Display, TEXT("Timecode not found for video file: %s"), *InFilePath);
 		return MakeError(Result.StealError());
 	}
 
+	UE_LOG(LogCaptureExtractTimecode, Display, TEXT("Extracted timecode '%s' from video file: %s"), *VideoInfoExtractor.GetTimecode().ToString(), *InFilePath);
+
 	return MakeValue(MoveTemp(VideoInfoExtractor));
 }
 
@@ -68,33 +72,34 @@ FCaptureExtractVideoInfo::FExtractResult FCaptureExtractVideoInfo::ExtractInfo(c
 		return MakeError(ECaptureExtractInfoError::InternalError);
 	}
 
-	FCaptureExtractVideoInfo::FExtractResult Result = ExtractInfoUsingElectraPlayer(InFilePath);;
+	const UCaptureManagerSettings* Settings = GetDefault<UCaptureManagerSettings>();
 
-	if (Result.HasError())
+	if (Settings->bEnableThirdPartyEncoder)
 	{
-		const UCaptureManagerSettings* Settings = GetDefault<UCaptureManagerSettings>();
-
-		if (Settings->bEnableThirdPartyEncoder)
+		FString EncoderPath = Settings->ThirdPartyEncoder.FilePath;
+		if (EncoderPath.EndsWith(TEXT("ffmpeg.exe")))
 		{
-			FString EncoderPath = Settings->ThirdPartyEncoder.FilePath;
-			if (EncoderPath.EndsWith(TEXT("ffmpeg.exe")))
+			EncoderPath.ReplaceInline(TEXT("ffmpeg.exe"), TEXT("ffprobe.exe"));
+
+			if (FPaths::FileExists(EncoderPath))
+			{
+				FExtractResult Result = ExtractInfoUsingFFProbe(InFilePath, EncoderPath);
+				return Result;
+			}
+			else
 			{
-				EncoderPath.ReplaceInline(TEXT("ffmpeg.exe"), TEXT("ffprobe.exe"));
-				Result = ExtractInfoUsingFFProbe(InFilePath, EncoderPath);
+				UE_LOG(LogCaptureExtractTimecode, Warning, TEXT("Could not extract timecode for '%s' using ffprobe. ffprobe could not be found at path %s. ffprobe is expected to be found in the same directory as ffmpeg"), *InFilePath, *EncoderPath);
+				return MakeError(ECaptureExtractInfoError::FfprobeNotFound);
 			}
 		}
 	}
-
-	if (Result.HasValue())
-	{
-		UE_LOG(LogCaptureExtractTimecode, Display, TEXT("Extracted timecode '%s' from video file: %s"), *TimecodeInfo.Timecode.ToString(), *FilePath);
-	}
 	else
 	{
-		UE_LOG(LogCaptureExtractTimecode, Display, TEXT("Timecode not found for video file: %s"), *FilePath);
+		FExtractResult Result = ExtractInfoUsingElectraPlayer(InFilePath);
+		return Result;
 	}
 
-	return Result;
+	return MakeError(ECaptureExtractInfoError::InternalError);
 }
 
 FCaptureExtractVideoInfo::FExtractResult FCaptureExtractVideoInfo::ExtractInfoUsingElectraPlayer(const FString& InFilePath)
diff --git a/Engine/Plugins/VirtualProduction/CaptureManager/CaptureManagerApp/Source/IngestLiveLinkDevice/Public/Utils/CaptureExtractTimecode.h b/Engine/Plugins/VirtualProduction/CaptureManager/CaptureManagerApp/Source/IngestLiveLinkDevice/Public/Utils/CaptureExtractTimecode.h
--- a/Engine/Plugins/VirtualProduction/CaptureManager/CaptureManagerApp/Source/IngestLiveLinkDevice/Public/Utils/CaptureExtractTimecode.h
+++ b/Engine/Plugins/VirtualProduction/CaptureManager/CaptureManagerApp/Source/IngestLiveLinkDevice/Public/Utils/CaptureExtractTimecode.h
@@ -41,6 +41,7 @@ enum ECaptureExtractInfoError : int32
 	TimecodeRateNotFound = 4,
 	UnhandledMedia = 5,
 	UnableToOpenMedia = 6,
+	FfprobeNotFound = 7,
 };
 }
 
diff --git a/Engine/Plugins/VirtualProduction/CaptureManager/CaptureManagerDevices/Source/VideoLiveLinkDeviceCommon/Private/Utils/VideoDeviceThumbnailExtractor.cpp b/Engine/Plugins/VirtualProduction/CaptureManager/CaptureManagerDevices/Source/VideoLiveLinkDeviceCommon/Private/Utils/VideoDeviceThumbnailExtractor.cpp
--- a/Engine/Plugins/VirtualProduction/CaptureManager/CaptureManagerDevices/Source/VideoLiveLinkDeviceCommon/Private/Utils/VideoDeviceThumbnailExtractor.cpp
+++ b/Engine/Plugins/VirtualProduction/CaptureManager/CaptureManagerDevices/Source/VideoLiveLinkDeviceCommon/Private/Utils/VideoDeviceThumbnailExtractor.cpp
@@ -231,6 +231,12 @@ TOptional<FIntPoint> FVideoDeviceThumbnailExtractor::ObtainImageDimensionsFromTh
 	FString CommandLine = FString::Format(TEXT("-v error -select_streams v:0 -show_entries stream=width,height -of default=noprint_wrappers=1 \"{0}\""), { InCurrentFile });
 	FString EncoderPath = Private::GetFFprobePath(InEncoderPath);
 
+	if (!FPaths::FileExists(EncoderPath))
+	{
+		UE_LOG(LogVideoLiveLinkDevice, Warning, TEXT("Could not extract image dimenstions for '%s' using ffprobe. ffprobe could not be found at path %s. ffprobe is expected to be found in the same directory as ffmpeg"), *InCurrentFile, *InEncoderPath);
+		return {};
+	}
+
 	TArray<uint8> WidthAndHeight = RunProcess(EncoderPath, CommandLine);
 
 	if (WidthAndHeight.IsEmpty())
