Fix launching of ffprobe
#Synced-CL 48507103

diff --git a/Engine/Plugins/VirtualProduction/CaptureManager/CaptureManagerApp/Source/IngestLiveLinkDevice/Private/Utils/CaptureExtractTimecode.cpp b/Engine/Plugins/VirtualProduction/CaptureManager/CaptureManagerApp/Source/IngestLiveLinkDevice/Private/Utils/CaptureExtractTimecode.cpp
--- a/Engine/Plugins/VirtualProduction/CaptureManager/CaptureManagerApp/Source/IngestLiveLinkDevice/Private/Utils/CaptureExtractTimecode.cpp
+++ b/Engine/Plugins/VirtualProduction/CaptureManager/CaptureManagerApp/Source/IngestLiveLinkDevice/Private/Utils/CaptureExtractTimecode.cpp
@@ -77,9 +77,13 @@ FCaptureExtractVideoInfo::FExtractResult FCaptureExtractVideoInfo::ExtractInfo(c
 	if (Settings->bEnableThirdPartyEncoder)
 	{
 		FString EncoderPath = Settings->ThirdPartyEncoder.FilePath;
-		if (EncoderPath.EndsWith(TEXT("ffmpeg.exe")))
+		FString BinarySuffix;
+#if PLATFORM_WINDOWS
+		BinarySuffix = TEXT(".exe"); // Temporary means of resolving ffprobe binary for 5.7.2. Refactor in 5.8 to use common code with VideoDeviceThumbnailExtractor.cpp
+#endif
+		if (EncoderPath.EndsWith(TEXT("ffmpeg") + BinarySuffix))
 		{
-			EncoderPath.ReplaceInline(TEXT("ffmpeg.exe"), TEXT("ffprobe.exe"));
+			EncoderPath.ReplaceInline(*(TEXT("ffmpeg") + BinarySuffix), *(TEXT("ffprobe") + BinarySuffix));
 
 			if (FPaths::FileExists(EncoderPath))
 			{
diff --git a/Engine/Plugins/VirtualProduction/CaptureManager/CaptureManagerDevices/Source/VideoLiveLinkDeviceCommon/Private/Utils/VideoDeviceThumbnailExtractor.cpp b/Engine/Plugins/VirtualProduction/CaptureManager/CaptureManagerDevices/Source/VideoLiveLinkDeviceCommon/Private/Utils/VideoDeviceThumbnailExtractor.cpp
--- a/Engine/Plugins/VirtualProduction/CaptureManager/CaptureManagerDevices/Source/VideoLiveLinkDeviceCommon/Private/Utils/VideoDeviceThumbnailExtractor.cpp
+++ b/Engine/Plugins/VirtualProduction/CaptureManager/CaptureManagerDevices/Source/VideoLiveLinkDeviceCommon/Private/Utils/VideoDeviceThumbnailExtractor.cpp
@@ -26,7 +26,7 @@ static bool IsEncoderFFmpeg(const FString& InEncoderPath)
 	return EncoderFileName == TEXT("ffmpeg");
 }
 
-static FString GetFFprobePath(const FString& InEncoderPath)
+static FString GetFFprobePath(const FString& InEncoderPath) // Centralize this in 5.8 for use with CaptureExtractTimecode.cpp
 {
 	check(IsEncoderFFmpeg(InEncoderPath));
 
