Capture Manager: Use the user settings directory as the default working and download directories instead of the user documents directory.

[RN] minor MetaHuman

#Synced-CL 48514830

diff --git a/Engine/Plugins/VirtualProduction/CaptureManager/CaptureManagerApp/Source/CaptureManagerSettings/Private/Settings/CaptureManagerSettings.cpp b/Engine/Plugins/VirtualProduction/CaptureManager/CaptureManagerApp/Source/CaptureManagerSettings/Private/Settings/CaptureManagerSettings.cpp
--- a/Engine/Plugins/VirtualProduction/CaptureManager/CaptureManagerApp/Source/CaptureManagerSettings/Private/Settings/CaptureManagerSettings.cpp
+++ b/Engine/Plugins/VirtualProduction/CaptureManager/CaptureManagerApp/Source/CaptureManagerSettings/Private/Settings/CaptureManagerSettings.cpp
@@ -72,12 +72,12 @@ void UCaptureManagerSettings::InitializeValuesIfNotSet()
 
 	if (DefaultWorkingDirectory.Path.IsEmpty())
 	{
-		DefaultWorkingDirectory.Path = FPaths::Combine(FPlatformProcess::UserDir(), Private::DefaultWorkingDirectory);
+		DefaultWorkingDirectory.Path = FPaths::Combine(FPlatformProcess::UserSettingsDir(), Private::DefaultWorkingDirectory);
 	}
 
 	if (DownloadDirectory.Path.IsEmpty())
 	{
-		DownloadDirectory.Path = FPaths::Combine(FPlatformProcess::UserDir(), Private::DefaultDownloadDirectory);
+		DownloadDirectory.Path = FPaths::Combine(FPlatformProcess::UserSettingsDir(), Private::DefaultDownloadDirectory);
 	}
 
 	if (DefaultUploadHostName.IsEmpty())
