Remove invalid warning message and generally reduce log spam when gathering timecode and frame rate data

When gathering frame rate and timecode information using ffprobe we always use the video frame rate for the timecode frame rate.
This warning message was copied from the electra player implementation, but it isn't relevant when using ffprobe and results in a log spam when using a third party encoder.

Also adjusts a few logs to make use of the 'Log' level which still appears within the editor log.

#rnx

#Synced-CL 48783767

diff --git a/Engine/Plugins/VirtualProduction/CaptureManager/CaptureManagerApp/Source/IngestLiveLinkDevice/Private/Utils/CaptureExtractTimecode.cpp b/Engine/Plugins/VirtualProduction/CaptureManager/CaptureManagerApp/Source/IngestLiveLinkDevice/Private/Utils/CaptureExtractTimecode.cpp
--- a/Engine/Plugins/VirtualProduction/CaptureManager/CaptureManagerApp/Source/IngestLiveLinkDevice/Private/Utils/CaptureExtractTimecode.cpp
+++ b/Engine/Plugins/VirtualProduction/CaptureManager/CaptureManagerApp/Source/IngestLiveLinkDevice/Private/Utils/CaptureExtractTimecode.cpp
@@ -51,11 +51,10 @@ FCaptureExtractVideoInfo::FResult FCaptureExtractVideoInfo::Create(const FString
 
 	if (Result.HasError())
 	{
-		UE_LOG(LogCaptureExtractTimecode, Display, TEXT("Timecode not found for video file: %s"), *InFilePath);
 		return MakeError(Result.StealError());
 	}
 
-	UE_LOG(LogCaptureExtractTimecode, Display, TEXT("Extracted timecode '%s' from video file: %s"), *VideoInfoExtractor.GetTimecode().ToString(), *InFilePath);
+	UE_LOG(LogCaptureExtractTimecode, Verbose, TEXT("Extracted timecode '%s' from video file: %s"), *VideoInfoExtractor.GetTimecode().ToString(), *InFilePath);
 
 	return MakeValue(MoveTemp(VideoInfoExtractor));
 }
@@ -420,7 +419,6 @@ FCaptureExtractVideoInfo::FExtractResult FCaptureExtractVideoInfo::ExtractInfoUs
 
 	if (bFoundFrameRate)
 	{
-		UE_LOG(LogCaptureExtractTimecode, Warning, TEXT("Timecode rate has not been found. Using video frame rate as timecode rate"));
 		TimecodeInfo.TimecodeRate = FrameRate;
 	}
 
@@ -534,11 +532,11 @@ FCaptureExtractAudioTimecode::FTimecodeInfoResult FCaptureExtractAudioTimecode::
 
 	if (Result.HasValue())
 	{
-		UE_LOG(LogCaptureExtractTimecode, Display, TEXT("Extracted timecode '%s' from audio file: %s"), *Result.GetValue().Timecode.ToString(), *FilePath);
+		UE_LOG(LogCaptureExtractTimecode, Verbose, TEXT("Extracted timecode '%s' from audio file: %s"), *Result.GetValue().Timecode.ToString(), *FilePath);
 	}
 	else
 	{
-		UE_LOG(LogCaptureExtractTimecode, Display, TEXT("Timecode not found for audio file: %s"), *FilePath);
+		UE_LOG(LogCaptureExtractTimecode, Log, TEXT("Timecode not found for audio file: %s"), *FilePath);
 	}
 
 	return Result;
@@ -565,7 +563,7 @@ FCaptureExtractAudioTimecode::FTimecodeInfoResult FCaptureExtractAudioTimecode::
 				{
 					UE_LOG(
 						LogCaptureExtractTimecode,
-						Display,
+						Log,
 						TEXT(
 							"Embedded timecode rate is %.2f fps (the sample rate). "
 							"This usually indicates there is no timecode rate information in the wav file: %s"
@@ -581,7 +579,7 @@ FCaptureExtractAudioTimecode::FTimecodeInfoResult FCaptureExtractAudioTimecode::
 
 						UE_LOG(
 							LogCaptureExtractTimecode,
-							Display,
+							Log,
 							TEXT(
 								"Taking the embedded audio timecode but estimating an SMPTE audio timecode rate. "
 								"Timecode rate for %s set to %.2f"
diff --git a/Engine/Plugins/VirtualProduction/CaptureManager/CaptureManagerDevices/Source/MonoVideoIngestDevice/Private/MonoVideoIngestDevice.cpp b/Engine/Plugins/VirtualProduction/CaptureManager/CaptureManagerDevices/Source/MonoVideoIngestDevice/Private/MonoVideoIngestDevice.cpp
--- a/Engine/Plugins/VirtualProduction/CaptureManager/CaptureManagerDevices/Source/MonoVideoIngestDevice/Private/MonoVideoIngestDevice.cpp
+++ b/Engine/Plugins/VirtualProduction/CaptureManager/CaptureManagerDevices/Source/MonoVideoIngestDevice/Private/MonoVideoIngestDevice.cpp
@@ -311,10 +311,6 @@ TOptional<FTakeMetadata> UMonoVideoIngestDevice::ReadTake(const FString& InCurre
 		TUniquePtr<IVideoReader> VideoReader = VideoReaderResult.StealValue();
 		Video.FrameRate = static_cast<float>(VideoReader->GetFrameRate().AsDecimal());
 	}
-	else
-	{
-		UE_LOG(LogVideoLiveLinkDevice, Warning, TEXT("Failed to determine the frame rate for the video file %s. Consider enabling Third Party Encoder in Capture Manager settings."), *InCurrentTakeFile);
-	}
 
 	FVideoDeviceThumbnailExtractor ThumbnailExtractor;
 	TOptional<FTakeThumbnailData::FRawImage> RawImageOpt = ThumbnailExtractor.ExtractThumbnail(InCurrentTakeFile);
@@ -340,6 +336,16 @@ TOptional<FTakeMetadata> UMonoVideoIngestDevice::ReadTake(const FString& InCurre
 		}
 	}
 
+	if (!Video.TimecodeStart.IsSet())
+	{
+		UE_LOG(LogVideoLiveLinkDevice, Warning, TEXT("Failed to determine the timecode for the video file %s."), *InCurrentTakeFile);
+	}
+
+	if (FMath::IsNearlyZero(Video.FrameRate))
+	{
+		UE_LOG(LogVideoLiveLinkDevice, Warning, TEXT("Failed to determine frame rate for video file: %s"), *InCurrentTakeFile);
+	}
+
 	TakeMetadata.Video.Add(MoveTemp(Video));
 
 	// Check if the video file contains audio stream by creating a audio reader from the file
diff --git a/Engine/Plugins/VirtualProduction/CaptureManager/CaptureManagerDevices/Source/StereoVideoIngestDevice/Private/StereoVideoIngestDevice.cpp b/Engine/Plugins/VirtualProduction/CaptureManager/CaptureManagerDevices/Source/StereoVideoIngestDevice/Private/StereoVideoIngestDevice.cpp
--- a/Engine/Plugins/VirtualProduction/CaptureManager/CaptureManagerDevices/Source/StereoVideoIngestDevice/Private/StereoVideoIngestDevice.cpp
+++ b/Engine/Plugins/VirtualProduction/CaptureManager/CaptureManagerDevices/Source/StereoVideoIngestDevice/Private/StereoVideoIngestDevice.cpp
@@ -592,10 +592,6 @@ TOptional<FTakeMetadata> UStereoVideoIngestDevice::CreateTakeMetadata(const FTak
 					bVideoFrameRateSet = true;
 				}
 			}
-			else
-			{
-				UE_LOG(LogVideoLiveLinkDevice, Warning, TEXT("Failed to determine the frame rate for the video file %s. Consider enabling Third Party Encoder in Capture Manager settings."), *VideoComponent.Path);
-			}
 
 			if (!TakeMetadata.Thumbnail.GetRawImage().IsSet())
 			{
@@ -624,9 +620,15 @@ TOptional<FTakeMetadata> UStereoVideoIngestDevice::CreateTakeMetadata(const FTak
 					bVideoFrameRateSet = true;
 				}
 			}
-			else
+
+			if (!Video.TimecodeStart.IsSet())
+			{
+				UE_LOG(LogVideoLiveLinkDevice, Warning, TEXT("Failed to determine the timecode for the video file %s."), *VideoComponent.Path);
+			}
+
+			if (FMath::IsNearlyZero(Video.FrameRate))
 			{
-				UE_LOG(LogVideoLiveLinkDevice, Warning, TEXT("Failed to determine the timecode for the video file %s. Consider enabling Third Party Encoder in Capture Manager settings."), *VideoComponent.Path);
+				UE_LOG(LogVideoLiveLinkDevice, Warning, TEXT("Failed to determine frame rate for video file: %s"), *VideoComponent.Path);
 			}
 		}
 		else // ETakeComponentType::IMAGE_SEQUENCE
