Ensure timecode information is fetched from the dedicated timecode track when available

Fixes timecode reading for mov files recorded from LLF
format_tags uses all available timecode sources to determine a final timecode value, which could be the video stream, or a dedicated timecode track.

#rnx
#Synced-CL 48783767

diff --git a/Engine/Plugins/VirtualProduction/CaptureManager/CaptureManagerApp/Source/IngestLiveLinkDevice/Private/Utils/CaptureExtractTimecode.cpp b/Engine/Plugins/VirtualProduction/CaptureManager/CaptureManagerApp/Source/IngestLiveLinkDevice/Private/Utils/CaptureExtractTimecode.cpp
--- a/Engine/Plugins/VirtualProduction/CaptureManager/CaptureManagerApp/Source/IngestLiveLinkDevice/Private/Utils/CaptureExtractTimecode.cpp
+++ b/Engine/Plugins/VirtualProduction/CaptureManager/CaptureManagerApp/Source/IngestLiveLinkDevice/Private/Utils/CaptureExtractTimecode.cpp
@@ -291,7 +291,7 @@ FCaptureExtractVideoInfo::FExtractResult FCaptureExtractVideoInfo::ExtractInfoUs
 {
 	using namespace UE::CaptureManager;
 
-	FString CommandArgs = FString::Format(TEXT("-v error -select_streams v:0 -show_entries stream_tags=timecode:stream=r_frame_rate -of default=noprint_wrappers=1:nokey=1 \"{0}\""), { InFilePath });
+	FString CommandArgs = FString::Format(TEXT("-v error -select_streams v:0 -show_entries format_tags=timecode:stream=r_frame_rate -of default=noprint_wrappers=1:nokey=1 \"{0}\""), { InFilePath });
 
 	void* ReadPipe = nullptr;
 	void* WritePipe = nullptr;
