Ensure audio assets are generated for mono takes containing an audio track when using ffmpeg on Mac & Linux

#rnx

#Synced-CL 49103540

diff --git a/Engine/Plugins/VirtualProduction/CaptureManager/CaptureManagerApp/CaptureManagerApp.uplugin b/Engine/Plugins/VirtualProduction/CaptureManager/CaptureManagerApp/CaptureManagerApp.uplugin
--- a/Engine/Plugins/VirtualProduction/CaptureManager/CaptureManagerApp/CaptureManagerApp.uplugin
+++ b/Engine/Plugins/VirtualProduction/CaptureManager/CaptureManagerApp/CaptureManagerApp.uplugin
@@ -97,6 +97,10 @@
 		{
 			"Name": "ElectraPlayer",
 			"Enabled": true
+		},
+		{
+			"Name": "MJPEGDecoderElectra",
+			"Enabled": true
 		}
 	]
 }
\ No newline at end of file
diff --git a/Engine/Plugins/VirtualProduction/CaptureManager/CaptureManagerApp/Source/IngestLiveLinkDevice/IngestLiveLinkDevice.Build.cs b/Engine/Plugins/VirtualProduction/CaptureManager/CaptureManagerApp/Source/IngestLiveLinkDevice/IngestLiveLinkDevice.Build.cs
--- a/Engine/Plugins/VirtualProduction/CaptureManager/CaptureManagerApp/Source/IngestLiveLinkDevice/IngestLiveLinkDevice.Build.cs
+++ b/Engine/Plugins/VirtualProduction/CaptureManager/CaptureManagerApp/Source/IngestLiveLinkDevice/IngestLiveLinkDevice.Build.cs
@@ -22,6 +22,7 @@ public class IngestLiveLinkDevice : ModuleRules
 			"CaptureProtocolStack",
 			"CaptureDataUtils",
 			"CaptureDataConverter",
+			"CaptureManagerMediaRW",
 			"CaptureManagerUnrealEndpoint",
 			"CaptureUtils",
 			"Engine",
@@ -38,7 +39,8 @@ public class IngestLiveLinkDevice : ModuleRules
 			"CaptureManagerSettings",
 			"CoreUObject",
 			"UnrealEd",
-			"ElectraPlayerPlugin"
+			"ElectraPlayerPlugin",
+			"MJPEGDecoderElectra"
 		});
 
 		CppCompileWarningSettings.UnsafeTypeCastWarningLevel = WarningLevel.Error;
diff --git a/Engine/Plugins/VirtualProduction/CaptureManager/CaptureManagerApp/Source/IngestLiveLinkDevice/Private/Utils/CaptureExtractTimecode.cpp b/Engine/Plugins/VirtualProduction/CaptureManager/CaptureManagerApp/Source/IngestLiveLinkDevice/Private/Utils/CaptureExtractTimecode.cpp
--- a/Engine/Plugins/VirtualProduction/CaptureManager/CaptureManagerApp/Source/IngestLiveLinkDevice/Private/Utils/CaptureExtractTimecode.cpp
+++ b/Engine/Plugins/VirtualProduction/CaptureManager/CaptureManagerApp/Source/IngestLiveLinkDevice/Private/Utils/CaptureExtractTimecode.cpp
@@ -4,6 +4,7 @@
 
 #include "Audio.h"
 #include "Async/Async.h"
+#include "CaptureManagerMediaRWModule.h"
 #include "Misc/FileHelper.h"
 #include "MediaPlaylist.h"
 #include "Sound/SoundWaveTimecodeInfo.h"
@@ -15,6 +16,7 @@
 #include "IMediaOptions.h"
 #include "IMediaEventSink.h"
 #include "IMediaTracks.h"
+#include "Serialization/JsonSerializer.h"
 
 #include "ParseTakeUtils.h"
 
@@ -76,6 +78,21 @@ FFrameRate FCaptureExtractVideoInfo::GetTimecodeRate() const
 	return TimecodeInfo.TimecodeRate;
 }
 
+bool FCaptureExtractVideoInfo::ContainsAudio() const
+{
+	return bContainsAudio;
+}
+
+float FCaptureExtractVideoInfo::GetAudioDurationSeconds() const
+{
+	return AudioDurationSeconds;
+}
+
+EMediaOrientation FCaptureExtractVideoInfo::GetVideoOrientation() const
+{
+	return VideoOrientation;
+}
+
 FCaptureExtractVideoInfo::FCaptureExtractVideoInfo()
 	: FrameRate(1, 1)
 {
@@ -100,8 +117,7 @@ FCaptureExtractVideoInfo::FExtractResult FCaptureExtractVideoInfo::ExtractInfo(c
 		FString FFprobePath = GetFFprobePath(EncoderPath);
 		if (FPaths::FileExists(FFprobePath))
 		{
-			FExtractResult Result = ExtractInfoUsingFFProbe(InFilePath, FFprobePath);
-			return Result;
+			return ExtractInfoUsingFFProbe(InFilePath, FFprobePath);
 		}
 		else
 		{
@@ -111,8 +127,17 @@ FCaptureExtractVideoInfo::FExtractResult FCaptureExtractVideoInfo::ExtractInfo(c
 	}
 	else
 	{
-		FExtractResult Result = ExtractInfoUsingElectraPlayer(InFilePath);
-		return Result;
+		// We gather audio duration via an audio reader if available and the format is supported.
+		FMediaRWManager& MediaManager = FModuleManager::LoadModuleChecked<FCaptureManagerMediaRWModule>("CaptureManagerMediaRW").Get();
+		TValueOrError<TUniquePtr<IAudioReader>, FText> AudioReaderResult = MediaManager.CreateAudioReader(InFilePath);
+		if (AudioReaderResult.HasValue())
+		{
+			bContainsAudio = true;
+			AudioDurationSeconds = static_cast<float>(AudioReaderResult.StealValue()->GetDuration().GetTotalSeconds());
+		}
+
+		// Frame rate, and timecode data is read from the electra player
+		return ExtractInfoUsingElectraPlayer(InFilePath);
 	}
 }
 
@@ -293,7 +318,7 @@ FCaptureExtractVideoInfo::FExtractResult FCaptureExtractVideoInfo::ExtractInfoUs
 {
 	using namespace UE::CaptureManager;
 
-	FString CommandArgs = FString::Format(TEXT("-v error -select_streams v:0 -show_entries format_tags=timecode:stream=r_frame_rate -of default=noprint_wrappers=1:nokey=1 \"{0}\""), { InFilePath });
+	FString CommandArgs = FString::Format(TEXT("-v error -show_entries stream=codec_type,r_frame_rate,duration:format_tags=timecode:stream_side_data=rotation -of json \"{0}\""), { InFilePath });
 
 	FProcessRunnerResult RunnerResult = FProcessRunner::Run(InFFProbePath, CommandArgs, nullptr, TimeoutPeriod);
 	
@@ -302,57 +327,118 @@ FCaptureExtractVideoInfo::FExtractResult FCaptureExtractVideoInfo::ExtractInfoUs
 		return MakeError(ECaptureExtractInfoError::InternalError);
 	}
 
-	bool bFoundTimecode = false;
-	bool bFoundFrameRate = false;
-
 	TArray<uint8> CommandOutput = RunnerResult.StealValue();
-	FString CommandOutputStr = FString::ConstructFromPtrSize(reinterpret_cast<const UTF8CHAR*>(CommandOutput.GetData()), CommandOutput.Num());
+	FString JsonString = FString::ConstructFromPtrSize(reinterpret_cast<const UTF8CHAR*>(CommandOutput.GetData()), CommandOutput.Num());
 
-	// Split the timecode and frame rate lines from the output
-	TArray<FString> OutputLinesStr;
-	CommandOutputStr.ParseIntoArrayLines(OutputLinesStr);
+	TSharedPtr<FJsonObject> JsonObject;
+	if (!FJsonSerializer::Deserialize(TJsonReaderFactory<>::Create(JsonString), JsonObject))
+	{
+		UE_LOG(LogCaptureExtractTimecode, Error, TEXT("Failed to parse ffprobe output for file \"%s\". Ensure ffprobe is up to date."), *FPaths::GetCleanFilename(InFilePath));
+		return MakeError(ECaptureExtractInfoError::FailedToParseJson);
+	}
 
-	for (const FString& Line : OutputLinesStr)
+	// Streams
+	const TArray<TSharedPtr<FJsonValue>>* StreamsArray;
+	if (JsonObject->TryGetArrayField(TEXT("streams"), StreamsArray))
 	{
-		if (!bFoundFrameRate)
+		for (const TSharedPtr<FJsonValue>& Value : *StreamsArray)
 		{
-			TArray<FString> FrameRatePartsStr;
-			Line.ParseIntoArray(FrameRatePartsStr, TEXT("/"));
-			if (FrameRatePartsStr.Num() == 2)
+			TSharedPtr<FJsonObject> StreamObject = Value->AsObject();
+			if (StreamObject.IsValid())
 			{
-				int32 Numerator = FCString::Atoi(*FrameRatePartsStr[0]);
-				int32 Denominator = FCString::Atoi(*FrameRatePartsStr[1]);
+				FString CodecType;
+				if (!StreamObject->TryGetStringField(TEXT("codec_type"), CodecType)) 
+				{
+					continue;
+				}
 
-				FrameRate = FFrameRate(Numerator, Denominator);
-				bFoundFrameRate = true;
-				continue;
+				if (CodecType == TEXT("audio"))
+				{
+					// Audio Duration
+					float Duration = 0.0;
+					if (StreamObject->TryGetNumberField(TEXT("duration"), Duration))
+					{
+						bContainsAudio = true;
+						AudioDurationSeconds = Duration;
+					}
+				}
+				else if (CodecType == TEXT("video"))
+				{
+					// Frame Rate
+					FString FrameRateString;
+					if (StreamObject->TryGetStringField(TEXT("r_frame_rate"), FrameRateString)) 
+					{
+						TArray<FString> FrameRateComponents;
+						FrameRateString.ParseIntoArray(FrameRateComponents, TEXT("/"));
+						if (FrameRateComponents.Num() == 2)
+						{
+							int32 Numerator = FCString::Atoi(*FrameRateComponents[0]);
+							int32 Denominator = FCString::Atoi(*FrameRateComponents[1]);
+							FrameRate = FFrameRate(Numerator, Denominator);
+							TimecodeInfo.TimecodeRate = FrameRate;
+						}
+					}
+
+					// Video Orientation
+					int32 Rotation = 0;
+					const TArray<TSharedPtr<FJsonValue>>* SideDataArrayPointer;
+					if (StreamObject->TryGetArrayField(TEXT("side_data_list"), SideDataArrayPointer))
+					{
+						const TArray<TSharedPtr<FJsonValue>> SideDataArray = *SideDataArrayPointer;
+						if (SideDataArray.Num() > 0)
+						{
+							const TSharedPtr<FJsonObject>* SideDataObjectPointer;
+							if (SideDataArray[0]->TryGetObject(SideDataObjectPointer))
+							{
+								(*SideDataObjectPointer)->TryGetNumberField(TEXT("rotation"), Rotation);
+							}
+						}
+					}
+					switch (FMath::WrapExclusive(Rotation, 0, 360))
+					{
+						case 90:
+							VideoOrientation = EMediaOrientation::CW90;
+							break;
+						case 180:
+							VideoOrientation = EMediaOrientation::CW180;
+							break;
+						case 270:
+							VideoOrientation = EMediaOrientation::CW270;
+							break;
+						default:
+							VideoOrientation = EMediaOrientation::Original;
+					}
+				}
 			}
 		}
+	}
 
-		if (!bFoundTimecode)
+	// Timecode
+	const TSharedPtr<FJsonObject>* FormatObject;
+	if (JsonObject->TryGetObjectField(TEXT("format"), FormatObject)) 
+	{
+		const TSharedPtr<FJsonObject>* TagsObject;
+		if ((*FormatObject)->TryGetObjectField(TEXT("tags"), TagsObject)) 
 		{
-			TArray<FString> TimecodePartsStr;
-			Line.ParseIntoArray(TimecodePartsStr, TEXT(":")); // We don't support DF timecode
-			if (TimecodePartsStr.Num() == 4)
+			FString TimecodeString;
+			if ((*TagsObject)->TryGetStringField(TEXT("timecode"), TimecodeString))
 			{
-				int32 Hours = FCString::Atoi(*TimecodePartsStr[0]);
-				int32 Minutes = FCString::Atoi(*TimecodePartsStr[1]);
-				int32 Seconds = FCString::Atoi(*TimecodePartsStr[2]);
-				int32 Frames = FCString::Atoi(*TimecodePartsStr[3]);
-
-				constexpr bool bIsDropFrame = false;
-				TimecodeInfo.Timecode = FTimecode(Hours, Minutes, Seconds, Frames, bIsDropFrame);
-				bFoundTimecode = true;
-				continue;
+				TArray<FString> TimecodeComponents;
+				TimecodeString.ParseIntoArray(TimecodeComponents, TEXT(":")); // We don't support DF timecode
+				if (TimecodeComponents.Num() == 4)
+				{
+					int32 Hours = FCString::Atoi(*TimecodeComponents[0]);
+					int32 Minutes = FCString::Atoi(*TimecodeComponents[1]);
+					int32 Seconds = FCString::Atoi(*TimecodeComponents[2]);
+					int32 Frames = FCString::Atoi(*TimecodeComponents[3]);
+
+					constexpr bool bIsDropFrame = false;
+					TimecodeInfo.Timecode = FTimecode(Hours, Minutes, Seconds, Frames, bIsDropFrame);
+				}
 			}
 		}
 	}
 
-	if (bFoundFrameRate)
-	{
-		TimecodeInfo.TimecodeRate = FrameRate;
-	}
-
 	return MakeValue();
 }
 
diff --git a/Engine/Plugins/VirtualProduction/CaptureManager/CaptureManagerApp/Source/IngestLiveLinkDevice/Public/Utils/CaptureExtractTimecode.h b/Engine/Plugins/VirtualProduction/CaptureManager/CaptureManagerApp/Source/IngestLiveLinkDevice/Public/Utils/CaptureExtractTimecode.h
--- a/Engine/Plugins/VirtualProduction/CaptureManager/CaptureManagerApp/Source/IngestLiveLinkDevice/Public/Utils/CaptureExtractTimecode.h
+++ b/Engine/Plugins/VirtualProduction/CaptureManager/CaptureManagerApp/Source/IngestLiveLinkDevice/Public/Utils/CaptureExtractTimecode.h
@@ -42,6 +42,7 @@ enum ECaptureExtractInfoError : int32
 	UnhandledMedia = 5,
 	UnableToOpenMedia = 6,
 	FfprobeNotFound = 7,
+	FailedToParseJson = 8,
 };
 }
 
@@ -70,6 +71,9 @@ public:
 	UE_API FFrameRate GetFrameRate() const;
 	UE_API FTimecode GetTimecode() const;
 	UE_API FFrameRate GetTimecodeRate() const;
+	UE_API bool ContainsAudio() const;
+	UE_API float GetAudioDurationSeconds() const;
+	UE_API EMediaOrientation GetVideoOrientation() const;
 
 private:
 
@@ -86,6 +90,9 @@ private:
 
 	FFrameRate FrameRate;
 	UE::CaptureManager::FTimecodeInfo TimecodeInfo;
+	bool bContainsAudio = false;
+	float AudioDurationSeconds = 0.0f;
+	EMediaOrientation VideoOrientation = EMediaOrientation::Original;
 };
 
 class FCaptureExtractAudioTimecode final : public TSharedFromThis<FCaptureExtractAudioTimecode>
diff --git a/Engine/Plugins/VirtualProduction/CaptureManager/CaptureManagerDevices/Source/MonoVideoIngestDevice/Private/MonoVideoIngestDevice.cpp b/Engine/Plugins/VirtualProduction/CaptureManager/CaptureManagerDevices/Source/MonoVideoIngestDevice/Private/MonoVideoIngestDevice.cpp
--- a/Engine/Plugins/VirtualProduction/CaptureManager/CaptureManagerDevices/Source/MonoVideoIngestDevice/Private/MonoVideoIngestDevice.cpp
+++ b/Engine/Plugins/VirtualProduction/CaptureManager/CaptureManagerDevices/Source/MonoVideoIngestDevice/Private/MonoVideoIngestDevice.cpp
@@ -6,8 +6,6 @@
 #include "Settings/CaptureManagerSettings.h"
 #include "ImageUtils.h"
 
-#include "CaptureManagerMediaRWModule.h"
-
 #include "Misc/MessageDialog.h"
 #include "Utils/TakeDiscoveryExpressionParser.h"
 #include "Utils/VideoDeviceThumbnailExtractor.h"
@@ -36,11 +34,6 @@ static const TArray<FStringView> SupportedVideoFormats =
 };
 
 
-FMediaRWManager& GetMediaRWManager()
-{
-	return FModuleManager::LoadModuleChecked<FCaptureManagerMediaRWModule>("CaptureManagerMediaRW").Get();
-}
-
 }
 
 const UMonoVideoIngestDeviceSettings* UMonoVideoIngestDevice::GetSettings() const
@@ -246,8 +239,6 @@ TOptional<FTakeMetadata> UMonoVideoIngestDevice::ReadTake(const FString& InCurre
 {
 	using namespace UE::CaptureManager;
 
-	FMediaRWManager& MediaManager = UE::MonoVideoLiveLinkDevice::Private::GetMediaRWManager();
-
 	FTakeMetadata TakeMetadata;
 
 	FFileStatData FileData = IFileManager::Get().GetStatData(*InCurrentTakeFile);
@@ -315,43 +306,60 @@ TOptional<FTakeMetadata> UMonoVideoIngestDevice::ReadTake(const FString& InCurre
 	Video.Path = InCurrentTakeFile;
 	Video.PathType = FTakeMetadata::FVideo::EPathType::File;
 
-	FCaptureExtractVideoInfo::FResult ExtractorOpt = FCaptureExtractVideoInfo::Create(Video.Path);
-	if (ExtractorOpt.IsValid())
+	FCaptureExtractVideoInfo::FResult ExtractResult = FCaptureExtractVideoInfo::Create(Video.Path);
+
+	if (ExtractResult.HasValue())
+	{
+		FCaptureExtractVideoInfo VideoFileInfo = ExtractResult.StealValue();
+		FString TimecodeStartString = VideoFileInfo.GetTimecode().ToString();
+		Video.TimecodeStart = TimecodeStartString;
+		Video.FrameRate = static_cast<uint16>(VideoFileInfo.GetFrameRate().AsDecimal());
+
+		switch (VideoFileInfo.GetVideoOrientation())
+		{
+			case EMediaOrientation::CW90:
+				Video.Orientation = FTakeMetadata::FVideo::EOrientation::CW90;
+				break;
+			case EMediaOrientation::CW180:
+				Video.Orientation = FTakeMetadata::FVideo::EOrientation::CW180;
+				break;
+			case EMediaOrientation::CW270:
+				Video.Orientation = FTakeMetadata::FVideo::EOrientation::CW270;
+				break;
+			default:
+				Video.Orientation = FTakeMetadata::FVideo::EOrientation::Original;
+		}
+
+		if (VideoFileInfo.ContainsAudio())
+		{
+			FTakeMetadata::FAudio Audio;
+			Audio.Name = TEXT("audio");
+			Audio.Path = InCurrentTakeFile;
+			Audio.Duration = VideoFileInfo.GetAudioDurationSeconds();
+
+			Audio.TimecodeStart = TimecodeStartString;
+			Audio.TimecodeRate = static_cast<float>(Video.FrameRate);
+
+			TakeMetadata.Audio.Add(MoveTemp(Audio));
+		}
+	}
+	else
 	{
-		FCaptureExtractVideoInfo Extractor = ExtractorOpt.StealValue();
-		Video.TimecodeStart = Extractor.GetTimecode().ToString();
-		Video.FrameRate = static_cast<float>(Extractor.GetFrameRate().AsDecimal());
+		UE_LOG(LogVideoLiveLinkDevice, Error, TEXT("Failed to read metadata for video file \"%s\""), *FPaths::GetCleanFilename(InCurrentTakeFile));
 	}
 
 	if (!Video.TimecodeStart.IsSet())
 	{
-		UE_LOG(LogVideoLiveLinkDevice, Warning, TEXT("Failed to determine the timecode for the video file %s."), *InCurrentTakeFile);
+		UE_LOG(LogVideoLiveLinkDevice, Warning, TEXT("Failed to determine the timecode for the video file \"%s\"."), *FPaths::GetCleanFilename(InCurrentTakeFile));
 	}
 
 	if (FMath::IsNearlyZero(Video.FrameRate))
 	{
-		UE_LOG(LogVideoLiveLinkDevice, Warning, TEXT("Failed to determine frame rate for video file: %s"), *InCurrentTakeFile);
+		UE_LOG(LogVideoLiveLinkDevice, Warning, TEXT("Failed to determine frame rate for video file: \"%s\""), *FPaths::GetCleanFilename(InCurrentTakeFile));
 	}
 
 	TakeMetadata.Video.Add(Video);
 
-	// Check if the video file contains audio stream by creating a audio reader from the file
-	// If creation fails, we assume that there is no audio in the file
-	TValueOrError<TUniquePtr<IAudioReader>, FText> AudioReaderResult = MediaManager.CreateAudioReader(InCurrentTakeFile);
-	
-	if (AudioReaderResult.HasValue())
-	{
-		FTakeMetadata::FAudio Audio;
-		Audio.Name = TEXT("audio");
-		Audio.Path = InCurrentTakeFile;
-		Audio.Duration = static_cast<float>(AudioReaderResult.StealValue()->GetDuration().GetTotalSeconds());
-
-		Audio.TimecodeStart = Video.TimecodeStart;
-		Audio.TimecodeRate = static_cast<float>(Video.FrameRate);
-
-		TakeMetadata.Audio.Add(MoveTemp(Audio));
-	}
-
 	return TakeMetadata;
 }
 
diff --git a/Engine/Plugins/VirtualProduction/CaptureManager/CaptureManagerDevices/Source/StereoVideoIngestDevice/Private/StereoVideoIngestDevice.cpp b/Engine/Plugins/VirtualProduction/CaptureManager/CaptureManagerDevices/Source/StereoVideoIngestDevice/Private/StereoVideoIngestDevice.cpp
--- a/Engine/Plugins/VirtualProduction/CaptureManager/CaptureManagerDevices/Source/StereoVideoIngestDevice/Private/StereoVideoIngestDevice.cpp
+++ b/Engine/Plugins/VirtualProduction/CaptureManager/CaptureManagerDevices/Source/StereoVideoIngestDevice/Private/StereoVideoIngestDevice.cpp
@@ -49,11 +49,6 @@ namespace UE::StereoVideoLiveLinkDevice::Private
 		TEXTVIEW("png")
 	};
 
-	FMediaRWManager& GetMediaRWManager()
-	{
-		return FModuleManager::LoadModuleChecked<FCaptureManagerMediaRWModule>("CaptureManagerMediaRW").Get();
-	}
-
 	static constexpr int32 MaximumVideoFilesCountPerTake = 2;
 	static constexpr int32 MaximumAudioFilesCountPerTake = 1;
 
@@ -523,8 +518,6 @@ TOptional<FTakeMetadata> UStereoVideoIngestDevice::CreateTakeMetadata(const FTak
 {
 	using namespace UE::CaptureManager;
 
-	FMediaRWManager& MediaManager = UE::StereoVideoLiveLinkDevice::Private::GetMediaRWManager();
-
 	FTakeMetadata TakeMetadata;
 
 	int32 TakeNumber = InTakeWithComponents.TakeNumber;
@@ -593,14 +586,29 @@ TOptional<FTakeMetadata> UStereoVideoIngestDevice::CreateTakeMetadata(const FTak
 
 			FString MediaFilePath = FPaths::ConvertRelativePathToFull(InTakeWithComponents.TakePath, Video.Path);
 
-			FCaptureExtractVideoInfo::FResult ExtractorOpt = FCaptureExtractVideoInfo::Create(MediaFilePath);
-			if (ExtractorOpt.IsValid())
+			FCaptureExtractVideoInfo::FResult ExtractResult = FCaptureExtractVideoInfo::Create(MediaFilePath);
+			if (ExtractResult.IsValid())
 			{
-				FCaptureExtractVideoInfo Extractor = ExtractorOpt.StealValue();
-				Video.TimecodeStart = Extractor.GetTimecode().ToString();			
-				Video.FrameRate = static_cast<float>(Extractor.GetFrameRate().AsDecimal());
-				VideoFrameRate = Extractor.GetFrameRate();
+				FCaptureExtractVideoInfo VideoFileInfo = ExtractResult.StealValue();
+				Video.TimecodeStart = VideoFileInfo.GetTimecode().ToString();			
+				Video.FrameRate = static_cast<float>(VideoFileInfo.GetFrameRate().AsDecimal());
+				VideoFrameRate = VideoFileInfo.GetFrameRate();
 				bVideoFrameRateSet = true;
+
+				switch (VideoFileInfo.GetVideoOrientation())
+				{
+					case EMediaOrientation::CW90:
+						Video.Orientation = FTakeMetadata::FVideo::EOrientation::CW90;
+						break;
+					case EMediaOrientation::CW180:
+						Video.Orientation = FTakeMetadata::FVideo::EOrientation::CW180;
+						break;
+					case EMediaOrientation::CW270:
+						Video.Orientation = FTakeMetadata::FVideo::EOrientation::CW270;
+						break;
+					default:
+						Video.Orientation = FTakeMetadata::FVideo::EOrientation::Original;
+				}
 			}
 
 			if (!Video.TimecodeStart.IsSet())
diff --git a/Engine/Plugins/VirtualProduction/CaptureManager/CaptureManagerDevices/Source/VideoLiveLinkDeviceCommon/Private/Utils/VideoDeviceThumbnailExtractor.cpp b/Engine/Plugins/VirtualProduction/CaptureManager/CaptureManagerDevices/Source/VideoLiveLinkDeviceCommon/Private/Utils/VideoDeviceThumbnailExtractor.cpp
--- a/Engine/Plugins/VirtualProduction/CaptureManager/CaptureManagerDevices/Source/VideoLiveLinkDeviceCommon/Private/Utils/VideoDeviceThumbnailExtractor.cpp
+++ b/Engine/Plugins/VirtualProduction/CaptureManager/CaptureManagerDevices/Source/VideoLiveLinkDeviceCommon/Private/Utils/VideoDeviceThumbnailExtractor.cpp
@@ -176,7 +176,7 @@ TArray<FColor> FVideoDeviceThumbnailExtractor::ObtainThumbnailFromThirdPartyEnco
 {
 	FString ScaleArg = FString::Printf(TEXT("scale=%d:-1"), ThirdPartyEncoderThumbnailWidth);
 
-	FString CommandLineArgs = FString::Format(TEXT("-noautorotate -hide_banner -loglevel error -i \"{0}\" -map v:0 -vf \"thumbnail=2, {1}\" -frames:v 1 -f rawvideo -pix_fmt rgb24 -"), { InCurrentFile, ScaleArg });
+	FString CommandLineArgs = FString::Format(TEXT("-hide_banner -loglevel error -i \"{0}\" -map v:0 -vf \"thumbnail=2, {1}\" -frames:v 1 -f rawvideo -pix_fmt rgb24 -"), { InCurrentFile, ScaleArg });
 
 	FProcessRunnerResult RunnerResult = FProcessRunner::Run(InEncoderPath, CommandLineArgs);
 
