Discovery requester wrong unlock call

#Synced-CL 49986193

diff --git a/Engine/Plugins/VirtualProduction/CaptureManager/CaptureManagerApp/Source/CaptureManagerUnrealEndpoint/Private/DiscoveryRequester.cpp b/Engine/Plugins/VirtualProduction/CaptureManager/CaptureManagerApp/Source/CaptureManagerUnrealEndpoint/Private/DiscoveryRequester.cpp
--- a/Engine/Plugins/VirtualProduction/CaptureManager/CaptureManagerApp/Source/CaptureManagerUnrealEndpoint/Private/DiscoveryRequester.cpp
+++ b/Engine/Plugins/VirtualProduction/CaptureManager/CaptureManagerApp/Source/CaptureManagerUnrealEndpoint/Private/DiscoveryRequester.cpp
@@ -279,7 +279,7 @@ void FDiscoveryRequester::FImpl::HandleDiscoveryResponse(
 		Client.SetLastDiscoveryResponse(FPlatformTime::Seconds());
 
 		KnownClients.Emplace(Client);
-		Mutex.Unlock();
+		Lock.Unlock();
 
 		UE_LOG(LogLiveLinkHubDiscovery, Display, TEXT("New client discovered %s with IP address %s and endpoint ID: %s"), *Client.GetHostName(), *Client.GetIPAddress(), *Client.GetClientID().ToString());
 		ClientFoundDelegate.Broadcast(Client);
