Updating ExampleLiveLinkDevice module after API changes

#Synced-CL 51005072

diff --git a/Engine/Plugins/VirtualProduction/CaptureManager/CaptureManagerApp/Source/ExampleLiveLinkDevices/Private/NetworkIngest/ExampleNetworkIngestDevice.cpp b/Engine/Plugins/VirtualProduction/CaptureManager/CaptureManagerApp/Source/ExampleLiveLinkDevices/Private/NetworkIngest/ExampleNetworkIngestDevice.cpp
--- a/Engine/Plugins/VirtualProduction/CaptureManager/CaptureManagerApp/Source/ExampleLiveLinkDevices/Private/NetworkIngest/ExampleNetworkIngestDevice.cpp
+++ b/Engine/Plugins/VirtualProduction/CaptureManager/CaptureManagerApp/Source/ExampleLiveLinkDevices/Private/NetworkIngest/ExampleNetworkIngestDevice.cpp
@@ -62,7 +62,7 @@ FString UExampleNetworkIngestDevice::GetFullTakePath(UE::CaptureManager::FTakeId
 	return FString();
 }
 
-void UExampleNetworkIngestDevice::UpdateTakeList_Implementation(UIngestCapability_UpdateTakeListCallback* InCallback)
+void UExampleNetworkIngestDevice::RunUpdateTakeList(UIngestCapability_UpdateTakeListCallback* InCallback)
 {
 	// Query your device for a list of takes
 	// Populate an FTakeMetadata object per take
diff --git a/Engine/Plugins/VirtualProduction/CaptureManager/CaptureManagerApp/Source/ExampleLiveLinkDevices/Private/NetworkIngest/ExampleNetworkIngestDevice.h b/Engine/Plugins/VirtualProduction/CaptureManager/CaptureManagerApp/Source/ExampleLiveLinkDevices/Private/NetworkIngest/ExampleNetworkIngestDevice.h
--- a/Engine/Plugins/VirtualProduction/CaptureManager/CaptureManagerApp/Source/ExampleLiveLinkDevices/Private/NetworkIngest/ExampleNetworkIngestDevice.h
+++ b/Engine/Plugins/VirtualProduction/CaptureManager/CaptureManagerApp/Source/ExampleLiveLinkDevices/Private/NetworkIngest/ExampleNetworkIngestDevice.h
@@ -56,7 +56,7 @@ private:
 	virtual FString GetFullTakePath(UE::CaptureManager::FTakeId InTakeId) const override;
 
 	//~ Begin ULiveLinkDeviceCapability_Ingest interface
-	virtual void UpdateTakeList_Implementation(UIngestCapability_UpdateTakeListCallback* InCallback) override;
+	virtual void RunUpdateTakeList(UIngestCapability_UpdateTakeListCallback* InCallback) override;
 	virtual void RunConvertAndUploadTake(const UIngestCapability_ProcessHandle* InProcessHandle,
 							   const UIngestCapability_Options* InIngestOptions) override;
 	virtual void CancelIngestProcess_Implementation(const UIngestCapability_ProcessHandle* InProcessHandle) override;
