Support DF timecode on ingest process

#Synced-CL 51408020

diff --git a/Engine/Plugins/VirtualProduction/CaptureManager/CaptureManagerApp/Source/IngestLiveLinkDevice/Private/Utils/CaptureExtractTimecode.cpp b/Engine/Plugins/VirtualProduction/CaptureManager/CaptureManagerApp/Source/IngestLiveLinkDevice/Private/Utils/CaptureExtractTimecode.cpp
--- a/Engine/Plugins/VirtualProduction/CaptureManager/CaptureManagerApp/Source/IngestLiveLinkDevice/Private/Utils/CaptureExtractTimecode.cpp
+++ b/Engine/Plugins/VirtualProduction/CaptureManager/CaptureManagerApp/Source/IngestLiveLinkDevice/Private/Utils/CaptureExtractTimecode.cpp
@@ -59,7 +59,7 @@ FCaptureExtractVideoInfo::FResult FCaptureExtractVideoInfo::Create(const FString
 		return MakeError(Result.StealError());
 	}
 
-	UE_LOG(LogCaptureExtractTimecode, Verbose, TEXT("Extracted timecode '%s' from video file: %s"), *VideoInfoExtractor.GetTimecode().ToString(), *InFilePath);
+	UE_LOGF(LogCaptureExtractTimecode, Verbose, "Extracted timecode '%ls' from video file: %ls", *VideoInfoExtractor.GetTimecode().ToString(), *InFilePath);
 
 	return MakeValue(MoveTemp(VideoInfoExtractor));
 }
@@ -122,7 +122,7 @@ FCaptureExtractVideoInfo::FExtractResult FCaptureExtractVideoInfo::ExtractInfo(c
 		}
 		else
 		{
-			UE_LOG(LogCaptureExtractTimecode, Warning, TEXT("Could not extract timecode for '%s' using ffprobe. ffprobe could not be found at path %s. ffprobe is expected to be found in the same directory as ffmpeg"), *InFilePath, *EncoderPath);
+			UE_LOGF(LogCaptureExtractTimecode, Warning, "Could not extract timecode for '%ls' using ffprobe. ffprobe could not be found at path %ls. ffprobe is expected to be found in the same directory as ffmpeg", *InFilePath, *EncoderPath);
 			return MakeError(ECaptureExtractInfoError::FfprobeNotFound);
 		}
 	}
@@ -257,7 +257,7 @@ FCaptureExtractVideoInfo::FExtractResult FCaptureExtractVideoInfo::ExtractInfoUs
 			}
 			else
 			{
-				UE_LOG(LogCaptureExtractTimecode, Warning, TEXT("Failed to obtain the frame rate"));
+				UE_LOGF(LogCaptureExtractTimecode, Warning, "Failed to obtain the frame rate");
 			}
 
 			break;
@@ -286,7 +286,7 @@ FCaptureExtractVideoInfo::FExtractResult FCaptureExtractVideoInfo::ExtractInfoUs
 		const TOptional<FTimecode> Timecode = FTimecode::ParseTimecode(*TimecodeString);
 		if (!Timecode.IsSet())
 		{
-			UE_LOG(LogCaptureExtractTimecode, Warning, TEXT("Failed to parse the timecode"));
+			UE_LOGF(LogCaptureExtractTimecode, Warning, "Failed to parse the timecode");
 		}
 		else
 		{
@@ -295,7 +295,7 @@ FCaptureExtractVideoInfo::FExtractResult FCaptureExtractVideoInfo::ExtractInfoUs
 	}
 	else
 	{
-		UE_LOG(LogCaptureExtractTimecode, Warning, TEXT("Timecode has not been found"));
+		UE_LOGF(LogCaptureExtractTimecode, Warning, "Timecode has not been found");
 	}
 
 	if (!TimecodeRateString.IsEmpty())
@@ -308,7 +308,7 @@ FCaptureExtractVideoInfo::FExtractResult FCaptureExtractVideoInfo::ExtractInfoUs
 	}
 	else
 	{
-		UE_LOG(LogCaptureExtractTimecode, Warning, TEXT("Timecode rate has not been found. Using video Frame rate as timecode rate"));
+		UE_LOGF(LogCaptureExtractTimecode, Warning, "Timecode rate has not been found. Using video Frame rate as timecode rate");
 		TimecodeInfo.TimecodeRate = FrameRate;
 	}
 
@@ -334,7 +334,7 @@ FCaptureExtractVideoInfo::FExtractResult FCaptureExtractVideoInfo::ExtractInfoUs
 	TSharedPtr<FJsonObject> JsonObject;
 	if (!FJsonSerializer::Deserialize(TJsonReaderFactory<>::Create(JsonString), JsonObject))
 	{
-		UE_LOG(LogCaptureExtractTimecode, Error, TEXT("Failed to parse ffprobe output for file \"%s\". Ensure ffprobe is up to date."), *FPaths::GetCleanFilename(InFilePath));
+		UE_LOGF(LogCaptureExtractTimecode, Error, "Failed to parse ffprobe output for file \"%ls\". Ensure ffprobe is up to date.", *FPaths::GetCleanFilename(InFilePath));
 		return MakeError(ECaptureExtractInfoError::FailedToParseJson);
 	}
 
@@ -424,17 +424,15 @@ FCaptureExtractVideoInfo::FExtractResult FCaptureExtractVideoInfo::ExtractInfoUs
 			FString TimecodeString;
 			if ((*TagsObject)->TryGetStringField(TEXT("timecode"), TimecodeString))
 			{
-				TArray<FString> TimecodeComponents;
-				TimecodeString.ParseIntoArray(TimecodeComponents, TEXT(":")); // We don't support DF timecode
-				if (TimecodeComponents.Num() == 4)
+				TOptional<FTimecode> TimecodeOpt = FTimecode::ParseTimecode(*TimecodeString);
+				if (TimecodeOpt)
 				{
-					int32 Hours = FCString::Atoi(*TimecodeComponents[0]);
-					int32 Minutes = FCString::Atoi(*TimecodeComponents[1]);
-					int32 Seconds = FCString::Atoi(*TimecodeComponents[2]);
-					int32 Frames = FCString::Atoi(*TimecodeComponents[3]);
-
-					constexpr bool bIsDropFrame = false;
-					TimecodeInfo.Timecode = FTimecode(Hours, Minutes, Seconds, Frames, bIsDropFrame);
+					TimecodeInfo.Timecode = TimecodeOpt.GetValue();
+				}
+				else
+				{
+					UE_LOGF(LogCaptureExtractTimecode, Error, "Failed to parse timecode string %ls", *TimecodeString);
+					return MakeError(ECaptureExtractInfoError::UnableToParseTimecode);
 				}
 			}
 		}
@@ -524,11 +522,11 @@ FCaptureExtractAudioTimecode::FTimecodeInfoResult FCaptureExtractAudioTimecode::
 
 	if (Result.HasValue())
 	{
-		UE_LOG(LogCaptureExtractTimecode, Verbose, TEXT("Extracted timecode '%s' from audio file: %s"), *Result.GetValue().Timecode.ToString(), *FilePath);
+		UE_LOGF(LogCaptureExtractTimecode, Verbose, "Extracted timecode '%ls' from audio file: %ls", *Result.GetValue().Timecode.ToString(), *FilePath);
 	}
 	else
 	{
-		UE_LOG(LogCaptureExtractTimecode, Log, TEXT("Timecode not found for audio file: %s"), *FilePath);
+		UE_LOGF(LogCaptureExtractTimecode, Log, "Timecode not found for audio file: %ls", *FilePath);
 	}
 
 	return Result;
@@ -553,13 +551,11 @@ FCaptureExtractAudioTimecode::FTimecodeInfoResult FCaptureExtractAudioTimecode::
 				const bool bTimecodeRateIsSampleRate = TimecodeRate == FFrameRate(TimecodeInfo.NumSamplesPerSecond, 1);
 				if (bTimecodeRateIsSampleRate)
 				{
-					UE_LOG(
+					UE_LOGF(
 						LogCaptureExtractTimecode,
 						Log,
-						TEXT(
 							"Embedded timecode rate is %.2f fps (the sample rate). "
-							"This usually indicates there is no timecode rate information in the wav file: %s"
-						),
+							"This usually indicates there is no timecode rate information in the wav file: %ls",
 						TimecodeInfo.TimecodeRate.AsDecimal(),
 						*FilePath
 					);
@@ -569,13 +565,11 @@ FCaptureExtractAudioTimecode::FTimecodeInfoResult FCaptureExtractAudioTimecode::
 						// Use the provided timecode rate instead
 						TimecodeRate = InTimecodeRate;
 
-						UE_LOG(
+						UE_LOGF(
 							LogCaptureExtractTimecode,
 							Log,
-							TEXT(
 								"Taking the embedded audio timecode but estimating an SMPTE audio timecode rate. "
-								"Timecode rate for %s set to %.2f"
-							),
+								"Timecode rate for %ls set to %.2f",
 							*FilePath,
 							TimecodeRate.AsDecimal()
 						);
