Adding streaming capability to LiveLink Devices

#Synced-CL 51579518

diff --git a/Engine/Plugins/VirtualProduction/CaptureManager/CaptureManagerApp/CaptureManagerApp.uplugin b/Engine/Plugins/VirtualProduction/CaptureManager/CaptureManagerApp/CaptureManagerApp.uplugin
--- a/Engine/Plugins/VirtualProduction/CaptureManager/CaptureManagerApp/CaptureManagerApp.uplugin
+++ b/Engine/Plugins/VirtualProduction/CaptureManager/CaptureManagerApp/CaptureManagerApp.uplugin
@@ -101,6 +101,10 @@
 		{
 			"Name": "MJPEGDecoderElectra",
 			"Enabled": true
+		},
+		{
+			"Name": "MediaFrameworkUtilities",
+			"Enabled": true
 		}
 	]
 }
\ No newline at end of file
diff --git a/Engine/Plugins/VirtualProduction/CaptureManager/CaptureManagerApp/Source/LiveLinkCapabilities/LiveLinkCapabilities.Build.cs b/Engine/Plugins/VirtualProduction/CaptureManager/CaptureManagerApp/Source/LiveLinkCapabilities/LiveLinkCapabilities.Build.cs
--- a/Engine/Plugins/VirtualProduction/CaptureManager/CaptureManagerApp/Source/LiveLinkCapabilities/LiveLinkCapabilities.Build.cs
+++ b/Engine/Plugins/VirtualProduction/CaptureManager/CaptureManagerApp/Source/LiveLinkCapabilities/LiveLinkCapabilities.Build.cs
@@ -27,9 +27,17 @@ public class LiveLinkCapabilities : ModuleRules
 			"LiveLinkHubCaptureMessaging",
 			"SlateCore",
 			"Slate",
-			"CoreUObject"
+			"CoreUObject",
+			"MediaAssets",
+			"MediaFrameworkUtilities",
+			"Engine"
 		});
 
+		if (Target.bBuildEditor) 
+		{ 
+			PrivateDependencyModuleNames.Add("UnrealEd"); 
+		}
+
 		CppCompileWarningSettings.UnsafeTypeCastWarningLevel = WarningLevel.Error;
 	}
 }
diff --git a/Engine/Plugins/VirtualProduction/CaptureManager/CaptureManagerApp/Source/LiveLinkCapabilities/Private/Streaming/LiveLinkDeviceCapability_Streaming.cpp b/Engine/Plugins/VirtualProduction/CaptureManager/CaptureManagerApp/Source/LiveLinkCapabilities/Private/Streaming/LiveLinkDeviceCapability_Streaming.cpp
new file mode 100644
--- /dev/null
+++ b/Engine/Plugins/VirtualProduction/CaptureManager/CaptureManagerApp/Source/LiveLinkCapabilities/Private/Streaming/LiveLinkDeviceCapability_Streaming.cpp
@@ -0,0 +1,323 @@
+// Copyright Epic Games, Inc. All Rights Reserved.
+
+#include "Streaming/LiveLinkDeviceCapability_Streaming.h"
+
+#include "Profile/IMediaProfileManager.h"
+#include "Profile/MediaProfile.h"
+#include "Profile/MediaProfileSettings.h"
+#include "Profile/MediaProfilePlaybackManager.h"
+
+#include "Engine/Engine.h"
+#include "LiveLinkDeviceSubsystem.h"
+
+#include "MediaSource.h"
+
+#include "Containers/Ticker.h"
+
+#include "ScopedTransaction.h"
+
+#define LOCTEXT_NAMESPACE "LiveLinkDeviceCapability_Streaming"
+
+void ULiveLinkDeviceCapability_Streaming::OnDeviceSubsystemInitialized()
+{
+	ULiveLinkDeviceSubsystem* Subsystem = GEngine->GetEngineSubsystem<ULiveLinkDeviceSubsystem>();
+
+	if (Subsystem)
+	{
+		Subsystem->OnDeviceAdded().AddUObject(this, &ULiveLinkDeviceCapability_Streaming::OnDeviceAdded);
+		Subsystem->OnDeviceRemoved().AddUObject(this, &ULiveLinkDeviceCapability_Streaming::OnDeviceRemoved);
+	}
+}
+
+void ULiveLinkDeviceCapability_Streaming::OnDeviceSubsystemDeinitializing()
+{
+	ULiveLinkDeviceSubsystem* Subsystem = GEngine->GetEngineSubsystem<ULiveLinkDeviceSubsystem>();
+
+	if (Subsystem)
+	{
+		Subsystem->OnDeviceAdded().RemoveAll(this);
+		Subsystem->OnDeviceRemoved().RemoveAll(this);
+	}
+}
+
+SHeaderRow::FColumn::FArguments& ULiveLinkDeviceCapability_Streaming::GenerateHeaderForColumn(const FName InColumnId, SHeaderRow::FColumn::FArguments& InArgs)
+{
+	return Super::GenerateHeaderForColumn(InColumnId, InArgs);
+}
+
+TSharedPtr<SWidget> ULiveLinkDeviceCapability_Streaming::GenerateWidgetForColumn(const FName InColumnId, const FLiveLinkDeviceWidgetArguments& InArgs, ULiveLinkDevice* InDevice)
+{
+	return Super::GenerateWidgetForColumn(InColumnId, InArgs, InDevice);
+}
+
+void ULiveLinkDeviceCapability_Streaming::OnDeviceAdded(FGuid InGuid, ULiveLinkDevice* InDevice)
+{
+	if (InDevice->Implements<ULiveLinkDeviceCapability_Streaming>())
+	{
+		TScriptInterface<ILiveLinkDeviceCapability_Streaming> Device(InDevice);
+		Device->Init();
+	}
+}
+
+void ULiveLinkDeviceCapability_Streaming::OnDeviceRemoved(FGuid InGuid, ULiveLinkDevice* InDevice)
+{
+	if (InDevice->Implements<ULiveLinkDeviceCapability_Streaming>())
+	{
+		TScriptInterface<ILiveLinkDeviceCapability_Streaming> Device(InDevice);
+		Device->Deinit();
+	}
+}
+
+class ILiveLinkDeviceCapability_Streaming::FImpl : public TSharedFromThis<ILiveLinkDeviceCapability_Streaming::FImpl>
+{
+public:
+
+	FImpl() = default;
+
+	void Init();
+	void Deinit();
+
+	UMediaSource* GetMediaSource() const;
+	void SetMediaSource(UMediaSource* InSource);
+	void RegisterMediaSource(UMediaSource* InMediaSource);
+	void UnregisterMediaSource();
+
+private:
+
+	void OnMediaProfileChanged(UMediaProfile* InPrevious, UMediaProfile* InNew);
+
+	void RegisterMediaSource_GameThread(UMediaSource* InMediaSource, UMediaProfile* InMediaProfile);
+	void UnregisterMediaSource_GameThread(UMediaProfile* InMediaProfile);
+
+	void RemoveCurrentMediaSource(UMediaProfile* InMediaProfile);
+
+	mutable FCriticalSection Mutex;
+	TStrongObjectPtr<UMediaSource> CurrentMediaSource = nullptr;
+
+	FDelegateHandle DelegateHandle;
+};
+
+void ILiveLinkDeviceCapability_Streaming::FImpl::Init()
+{
+	IMediaProfileManager& MediaProfileManager = IMediaProfileManager::Get();
+	DelegateHandle = MediaProfileManager.OnMediaProfileChanged().AddSP(this, &ILiveLinkDeviceCapability_Streaming::FImpl::OnMediaProfileChanged);
+}
+
+void ILiveLinkDeviceCapability_Streaming::FImpl::Deinit()
+{
+	IMediaProfileManager& MediaProfileManager = IMediaProfileManager::Get();
+	MediaProfileManager.OnMediaProfileChanged().Remove(DelegateHandle);
+}
+
+UMediaSource* ILiveLinkDeviceCapability_Streaming::FImpl::GetMediaSource() const
+{
+	FScopeLock Lock(&Mutex);
+	return CurrentMediaSource.Get();
+}
+
+void ILiveLinkDeviceCapability_Streaming::FImpl::SetMediaSource(UMediaSource* InSource)
+{
+	FScopeLock Lock(&Mutex);
+	CurrentMediaSource.Reset(InSource);
+}
+
+void ILiveLinkDeviceCapability_Streaming::FImpl::RegisterMediaSource(UMediaSource* InMediaSource)
+{
+	TStrongObjectPtr<UMediaSource> MediaSource(InMediaSource);
+
+	if (IsInGameThread())
+	{
+		IMediaProfileManager& MediaProfileManager = IMediaProfileManager::Get();
+		UMediaProfile* MediaProfile = MediaProfileManager.GetCurrentMediaProfile();
+
+		RegisterMediaSource_GameThread(MediaSource.Get(), MediaProfile);
+
+		SetMediaSource(MediaSource.Get());
+		MediaSource.Reset();
+	}
+	else
+	{
+		ExecuteOnGameThread(TEXT("Updating Media Profile"), [This = AsWeak(), MediaSource = MoveTemp(MediaSource)]() mutable
+		{
+			if (TSharedPtr<FImpl> SharedThis = This.Pin())
+			{
+				IMediaProfileManager& MediaProfileManager = IMediaProfileManager::Get();
+				UMediaProfile* MediaProfile = MediaProfileManager.GetCurrentMediaProfile();
+
+				SharedThis->RegisterMediaSource_GameThread(MediaSource.Get(), MediaProfile);
+
+				SharedThis->SetMediaSource(MediaSource.Get());
+				MediaSource.Reset();
+			}
+		});
+	}
+}
+
+void ILiveLinkDeviceCapability_Streaming::FImpl::UnregisterMediaSource()
+{
+	if (IsInGameThread())
+	{
+		IMediaProfileManager& MediaProfileManager = IMediaProfileManager::Get();
+		UMediaProfile* MediaProfile = MediaProfileManager.GetCurrentMediaProfile();
+
+		UnregisterMediaSource_GameThread(MediaProfile);
+		SetMediaSource(nullptr);
+	}
+	else
+	{
+		ExecuteOnGameThread(TEXT("Updating Media Profile"), [This = AsWeak()]() mutable
+		{
+			if (TSharedPtr<FImpl> SharedThis = This.Pin())
+			{
+				IMediaProfileManager& MediaProfileManager = IMediaProfileManager::Get();
+				UMediaProfile* MediaProfile = MediaProfileManager.GetCurrentMediaProfile();
+
+				SharedThis->UnregisterMediaSource_GameThread(MediaProfile);
+				SharedThis->SetMediaSource(nullptr);
+			}
+		});
+	}
+}
+
+void ILiveLinkDeviceCapability_Streaming::FImpl::OnMediaProfileChanged(UMediaProfile* InPrevious, UMediaProfile* InNew)
+{
+	FScopeLock Lock(&Mutex);
+
+	if (!CurrentMediaSource.Get())
+	{
+		return;
+	}
+
+	if (InPrevious)
+	{
+		UnregisterMediaSource_GameThread(InPrevious);
+	}
+	
+	if (InNew)
+	{
+		RegisterMediaSource_GameThread(CurrentMediaSource.Get(), InNew);
+	}
+}
+
+void ILiveLinkDeviceCapability_Streaming::FImpl::RegisterMediaSource_GameThread(UMediaSource* InMediaSource, UMediaProfile* InMediaProfile)
+{
+	if (!InMediaProfile)
+	{
+		return;
+	}
+
+	FScopedTransaction UpdateMediaProfileTransaction(LOCTEXT("UpdateMediaProfileTransaction", "Update Media Profile"));
+	InMediaProfile->Modify();
+
+	RemoveCurrentMediaSource(InMediaProfile);
+
+	if (InMediaSource)
+	{
+		InMediaProfile->AddMediaSource(InMediaSource);
+	}
+
+	int32 NumOfSources = InMediaProfile->NumMediaSources();
+	GetMutableDefault<UMediaProfileSettings>()->FillDefaultMediaSourceProxies(NumOfSources, true);
+
+#if WITH_EDITOR
+	FProperty* Property = FindFProperty<FProperty>(UMediaProfile::StaticClass(), "MediaSources"); // MediaSources is protected field
+
+	if (Property)
+	{
+		FPropertyChangedEvent PropertyChangedEvent(Property, EPropertyChangeType::ArrayAdd);
+		InMediaProfile->PostEditChangeProperty(PropertyChangedEvent);
+	}
+#endif
+
+	InMediaProfile->MarkPackageDirty();
+}
+
+void ILiveLinkDeviceCapability_Streaming::FImpl::UnregisterMediaSource_GameThread(UMediaProfile* InMediaProfile)
+{
+	if (!InMediaProfile)
+	{
+		return;
+	}
+
+	FScopedTransaction RemoveSourceMediaProfileTransaction(LOCTEXT("RemoveSourceMediaProfileTransaction", "Remove Media Source"));
+	InMediaProfile->Modify();
+
+	{
+		// Close the source if we are unregistering
+		FScopeLock Lock(&Mutex);
+
+		UMediaProfilePlaybackManager::FCloseSourceArgs Args;
+		Args.bForceClose = true;
+		InMediaProfile->GetPlaybackManager()->CloseSource(CurrentMediaSource.Get(), Args);
+	}
+
+	RemoveCurrentMediaSource(InMediaProfile);
+
+	int32 NumOfSources = InMediaProfile->NumMediaSources();
+	GetMutableDefault<UMediaProfileSettings>()->FillDefaultMediaSourceProxies(NumOfSources, true);
+
+#if WITH_EDITOR
+	FProperty* Property = FindFProperty<FProperty>(UMediaProfile::StaticClass(), "MediaSources"); // MediaSources is protected field
+
+	if (Property)
+	{
+		FPropertyChangedEvent PropertyChangedEvent(Property, EPropertyChangeType::ArrayRemove);
+		InMediaProfile->PostEditChangeProperty(PropertyChangedEvent);
+	}
+#endif
+
+	InMediaProfile->MarkPackageDirty();
+}
+
+void ILiveLinkDeviceCapability_Streaming::FImpl::RemoveCurrentMediaSource(UMediaProfile* InMediaProfile)
+{
+	FScopeLock Lock(&Mutex);
+
+	int32 Index = InMediaProfile->FindMediaSourceIndex(CurrentMediaSource.Get());
+	if (Index != INDEX_NONE)
+	{
+		InMediaProfile->RemoveMediaSource(Index);
+	}
+}
+
+ILiveLinkDeviceCapability_Streaming::ILiveLinkDeviceCapability_Streaming()
+	: Impl(MakeShared<FImpl>())
+{
+}
+
+void ILiveLinkDeviceCapability_Streaming::Start_Implementation()
+{
+	Impl->RegisterMediaSource(Impl->GetMediaSource());
+}
+
+void ILiveLinkDeviceCapability_Streaming::Stop_Implementation()
+{
+	Impl->UnregisterMediaSource();
+}
+
+UMediaSource* ILiveLinkDeviceCapability_Streaming::GetMediaSource_Implementation() const
+{
+	return Impl->GetMediaSource();
+}
+
+void ILiveLinkDeviceCapability_Streaming::RegisterMediaSource(UMediaSource* InMediaSource)
+{
+	Impl->RegisterMediaSource(InMediaSource);
+}
+
+void ILiveLinkDeviceCapability_Streaming::UnregisterMediaSource()
+{
+	Impl->UnregisterMediaSource();
+}
+
+void ILiveLinkDeviceCapability_Streaming::Init()
+{
+	Impl->Init();
+}
+
+void ILiveLinkDeviceCapability_Streaming::Deinit()
+{
+	Impl->Deinit();
+}
+
+#undef LOCTEXT_NAMESPACE
diff --git a/Engine/Plugins/VirtualProduction/CaptureManager/CaptureManagerApp/Source/LiveLinkCapabilities/Public/Streaming/LiveLinkDeviceCapability_Streaming.h b/Engine/Plugins/VirtualProduction/CaptureManager/CaptureManagerApp/Source/LiveLinkCapabilities/Public/Streaming/LiveLinkDeviceCapability_Streaming.h
new file mode 100644
--- /dev/null
+++ b/Engine/Plugins/VirtualProduction/CaptureManager/CaptureManagerApp/Source/LiveLinkCapabilities/Public/Streaming/LiveLinkDeviceCapability_Streaming.h
@@ -0,0 +1,71 @@
+// Copyright Epic Games, Inc. All Rights Reserved.
+
+#pragma once
+
+#include "LiveLinkDeviceCapability.h"
+
+#include "LiveLinkDevice.h"
+
+#include "LiveLinkDeviceCapability_Streaming.generated.h"
+
+UINTERFACE(Blueprintable)
+class LIVELINKCAPABILITIES_API ULiveLinkDeviceCapability_Streaming :
+	public ULiveLinkDeviceCapability
+{
+	GENERATED_BODY()
+
+public:
+	
+	ULiveLinkDeviceCapability_Streaming() = default;
+
+	/** Called at completion of Live Link device subsystem initialization. */
+	virtual void OnDeviceSubsystemInitialized();
+
+	/** Called at the beginning of Live Link device subsystem de-initialization. */
+	virtual void OnDeviceSubsystemDeinitializing();
+
+	virtual SHeaderRow::FColumn::FArguments& GenerateHeaderForColumn(const FName InColumnId, SHeaderRow::FColumn::FArguments& InArgs) override;
+	virtual TSharedPtr<SWidget> GenerateWidgetForColumn(const FName InColumnId, const FLiveLinkDeviceWidgetArguments& InArgs, ULiveLinkDevice* InDevice) override;
+
+private:
+
+	void OnDeviceAdded(FGuid InGuid, ULiveLinkDevice* InDevice);
+	void OnDeviceRemoved(FGuid InGuid, ULiveLinkDevice* InDevice);
+};
+
+class LIVELINKCAPABILITIES_API ILiveLinkDeviceCapability_Streaming
+	: public ILiveLinkDeviceCapability
+{
+	GENERATED_BODY()
+
+public:
+
+	ILiveLinkDeviceCapability_Streaming();
+
+	UFUNCTION(BlueprintCallable, BlueprintNativeEvent, Category = "Live Link Device|Streaming")
+	void Start();
+	void Start_Implementation();
+
+	UFUNCTION(BlueprintCallable, BlueprintNativeEvent, Category = "Live Link Device|Streaming")
+	void Stop();
+	void Stop_Implementation();
+
+	UFUNCTION(BlueprintCallable, BlueprintNativeEvent, Category = "Live Link Device|Streaming")
+	class UMediaSource* GetMediaSource() const;
+	class UMediaSource* GetMediaSource_Implementation() const;
+
+protected:
+
+	void RegisterMediaSource(class UMediaSource* InMediaSource);
+	void UnregisterMediaSource();
+
+private:
+
+	void Init();
+	void Deinit();
+
+	friend class ULiveLinkDeviceCapability_Streaming;
+
+	class FImpl;
+	TSharedPtr<FImpl> Impl;
+};
\ No newline at end of file
diff --git a/Engine/Plugins/VirtualProduction/CaptureManager/CaptureManagerDevices/CaptureManagerDevices.uplugin b/Engine/Plugins/VirtualProduction/CaptureManager/CaptureManagerDevices/CaptureManagerDevices.uplugin
--- a/Engine/Plugins/VirtualProduction/CaptureManager/CaptureManagerDevices/CaptureManagerDevices.uplugin
+++ b/Engine/Plugins/VirtualProduction/CaptureManager/CaptureManagerDevices/CaptureManagerDevices.uplugin
@@ -51,6 +51,10 @@
 		{
 			"Name": "CaptureManagerCore",
 			"Enabled": true
+		},
+		{
+			"Name": "RTSPMedia",
+			"Enabled": true
 		}
 	]
 }
\ No newline at end of file
diff --git a/Engine/Plugins/VirtualProduction/CaptureManager/CaptureManagerDevices/Source/CPSLiveLinkDevice/CPSLiveLinkDevice.Build.cs b/Engine/Plugins/VirtualProduction/CaptureManager/CaptureManagerDevices/Source/CPSLiveLinkDevice/CPSLiveLinkDevice.Build.cs
--- a/Engine/Plugins/VirtualProduction/CaptureManager/CaptureManagerDevices/Source/CPSLiveLinkDevice/CPSLiveLinkDevice.Build.cs
+++ b/Engine/Plugins/VirtualProduction/CaptureManager/CaptureManagerDevices/Source/CPSLiveLinkDevice/CPSLiveLinkDevice.Build.cs
@@ -30,6 +30,7 @@ public class CPSLiveLinkDevice : ModuleRules
 			"LiveLinkHubCaptureMessaging",
 			"LiveLinkHub",
 			"LiveLinkFaceMetadata",
+			"RTSPMedia",
 			"ToolWidgets",
 			"Slate",
 			"SlateCore",
diff --git a/Engine/Plugins/VirtualProduction/CaptureManager/CaptureManagerDevices/Source/CPSLiveLinkDevice/Private/LiveLinkFaceDevice.cpp b/Engine/Plugins/VirtualProduction/CaptureManager/CaptureManagerDevices/Source/CPSLiveLinkDevice/Private/LiveLinkFaceDevice.cpp
--- a/Engine/Plugins/VirtualProduction/CaptureManager/CaptureManagerDevices/Source/CPSLiveLinkDevice/Private/LiveLinkFaceDevice.cpp
+++ b/Engine/Plugins/VirtualProduction/CaptureManager/CaptureManagerDevices/Source/CPSLiveLinkDevice/Private/LiveLinkFaceDevice.cpp
@@ -18,6 +18,7 @@
 #include "Utils/CaptureExtractTimecode.h"
 #include "Utils/ParseTakeUtils.h"
 
+#include "RtspMediaSource.h"
 
 DEFINE_LOG_CATEGORY_STATIC(LogLiveLinkFaceDevice, Log, All);
 
@@ -326,8 +327,18 @@ void ULiveLinkFaceDevice::HandleConnectionChanged(TSharedPtr<const UE::CaptureMa
 	}
 	else if (Event->ConnectionState == FConnectionStateChangedEvent::EState::Connected)
 	{
+		TObjectPtr<URtspMediaSource> MediaSource = NewObject<URtspMediaSource>();
+		MediaSource->Host = GetSettings()->IpAddress.IpAddressString;
+		MediaSource->DecoderBufferSize = 1; // Using 1 for minimum latency
+
+		RegisterMediaSource(MediaSource);
+
 		bIsConnecting = false;
 	}
+	else if (Event->ConnectionState == FConnectionStateChangedEvent::EState::Disconnected)
+	{
+		UnregisterMediaSource();
+	}
 
 	auto GetConnectionState = [](FConnectionStateChangedEvent::EState InConnectionState) -> ELiveLinkDeviceConnectionStatus
 	{
diff --git a/Engine/Plugins/VirtualProduction/CaptureManager/CaptureManagerDevices/Source/CPSLiveLinkDevice/Public/LiveLinkFaceDevice.h b/Engine/Plugins/VirtualProduction/CaptureManager/CaptureManagerDevices/Source/CPSLiveLinkDevice/Public/LiveLinkFaceDevice.h
--- a/Engine/Plugins/VirtualProduction/CaptureManager/CaptureManagerDevices/Source/CPSLiveLinkDevice/Public/LiveLinkFaceDevice.h
+++ b/Engine/Plugins/VirtualProduction/CaptureManager/CaptureManagerDevices/Source/CPSLiveLinkDevice/Public/LiveLinkFaceDevice.h
@@ -7,6 +7,7 @@
 #include "Ingest/LiveLinkDeviceCapability_Ingest.h"
 #include "LiveLinkDeviceCapability_Connection.h"
 #include "LiveLinkDeviceCapability_Recording.h"
+#include "Streaming/LiveLinkDeviceCapability_Streaming.h"
 #include "Misc/Optional.h"
 
 #include "Async/TaskProgress.h"
@@ -44,6 +45,7 @@ UCLASS(BlueprintType, meta = (DisplayName = "Live Link Face Device", ToolTip = "
 class CPSLIVELINKDEVICE_API ULiveLinkFaceDevice : public UBaseIngestLiveLinkDevice
 	, public ILiveLinkDeviceCapability_Connection
 	, public ILiveLinkDeviceCapability_Recording
+	, public ILiveLinkDeviceCapability_Streaming
 {
 	GENERATED_BODY()
 
