[CaptureManager] Moved conversion modules to CaptureManagerCore.

Moves CaptureManagerPipeline, CaptureManagerMediaRW, and CaptureDataConverter
from CaptureManagerApp (Editor type) to CaptureManagerCore (Editor type),
enabling both LiveLinkHub and UE Editor workflows to share identical conversion
logic from a single compiled module.

#Synced-CL 52057491

diff --git a/Engine/Plugins/VirtualProduction/CaptureManager/CaptureManagerApp/CaptureManagerApp.uplugin b/Engine/Plugins/VirtualProduction/CaptureManager/CaptureManagerApp/CaptureManagerApp.uplugin
--- a/Engine/Plugins/VirtualProduction/CaptureManager/CaptureManagerApp/CaptureManagerApp.uplugin
+++ b/Engine/Plugins/VirtualProduction/CaptureManager/CaptureManagerApp/CaptureManagerApp.uplugin
@@ -33,21 +33,6 @@
 			"Type": "Editor",
 			"LoadingPhase": "Default"
 		},
-		{
-			"Name": "CaptureManagerPipeline",
-			"Type": "Editor",
-			"LoadingPhase": "Default"
-		},
-		{
-			"Name": "CaptureManagerMediaRW",
-			"Type": "Editor",
-			"LoadingPhase": "Default"
-		},
-		{
-			"Name": "CaptureDataConverter",
-			"Type": "Editor",
-			"LoadingPhase": "Default"
-		},
 		{
 			"Name": "CaptureManagerUnrealEndpoint",
 			"Type": "Editor",
@@ -74,14 +59,6 @@
 			"Name": "LiveLinkHub",
 			"Enabled": true
 		},
-		{
-			"Name": "NamingTokens",
-			"Enabled": true
-		},
-		{
-			"Name": "ControlFlows",
-			"Enabled": true
-		},
 		{
 			"Name": "CaptureManagerCore",
 			"Enabled": true
diff --git a/Engine/Plugins/VirtualProduction/CaptureManager/CaptureManagerApp/Source/CaptureManagerSettings/CaptureManagerSettings.Build.cs b/Engine/Plugins/VirtualProduction/CaptureManager/CaptureManagerApp/Source/CaptureManagerSettings/CaptureManagerSettings.Build.cs
--- a/Engine/Plugins/VirtualProduction/CaptureManager/CaptureManagerApp/Source/CaptureManagerSettings/CaptureManagerSettings.Build.cs
+++ b/Engine/Plugins/VirtualProduction/CaptureManager/CaptureManagerApp/Source/CaptureManagerSettings/CaptureManagerSettings.Build.cs
@@ -24,7 +24,8 @@ public class CaptureManagerSettings : ModuleRules
 			"SlateCore",
 			"Settings",
 			"UnrealEd",
-			"CaptureUtils"
+			"CaptureUtils",
+			"CaptureDataConverter"
 		});
 	}
 }
diff --git a/Engine/Plugins/VirtualProduction/CaptureManager/CaptureManagerApp/Source/CaptureManagerSettings/Private/Settings/CaptureManagerSettings.cpp b/Engine/Plugins/VirtualProduction/CaptureManager/CaptureManagerApp/Source/CaptureManagerSettings/Private/Settings/CaptureManagerSettings.cpp
--- a/Engine/Plugins/VirtualProduction/CaptureManager/CaptureManagerApp/Source/CaptureManagerSettings/Private/Settings/CaptureManagerSettings.cpp
+++ b/Engine/Plugins/VirtualProduction/CaptureManager/CaptureManagerApp/Source/CaptureManagerSettings/Private/Settings/CaptureManagerSettings.cpp
@@ -2,6 +2,7 @@
 
 #include "Settings/CaptureManagerSettings.h"
 #include "Settings/CaptureManagerTemplateTokens.h"
+#include "CaptureManagerEncoderNamingTokens.h"
 #include "Network/NetworkMisc.h"
 #include "Misc/Paths.h"
 
diff --git a/Engine/Plugins/VirtualProduction/CaptureManager/CaptureManagerApp/Source/CaptureManagerSettings/Private/Settings/CaptureManagerSettingsCustomization.cpp b/Engine/Plugins/VirtualProduction/CaptureManager/CaptureManagerApp/Source/CaptureManagerSettings/Private/Settings/CaptureManagerSettingsCustomization.cpp
--- a/Engine/Plugins/VirtualProduction/CaptureManager/CaptureManagerApp/Source/CaptureManagerSettings/Private/Settings/CaptureManagerSettingsCustomization.cpp
+++ b/Engine/Plugins/VirtualProduction/CaptureManager/CaptureManagerApp/Source/CaptureManagerSettings/Private/Settings/CaptureManagerSettingsCustomization.cpp
@@ -4,6 +4,7 @@
 
 #include "Settings/CaptureManagerTemplateTokens.h"
 #include "Settings/CaptureManagerSettings.h"
+#include "CaptureManagerEncoderNamingTokens.h"
 #include "DetailLayoutBuilder.h"
 #include "DetailWidgetRow.h"
 #include "Engine/Engine.h"
diff --git a/Engine/Plugins/VirtualProduction/CaptureManager/CaptureManagerApp/Source/CaptureManagerSettings/Private/Settings/CaptureManagerTemplateTokens.cpp b/Engine/Plugins/VirtualProduction/CaptureManager/CaptureManagerApp/Source/CaptureManagerSettings/Private/Settings/CaptureManagerTemplateTokens.cpp
--- a/Engine/Plugins/VirtualProduction/CaptureManager/CaptureManagerApp/Source/CaptureManagerSettings/Private/Settings/CaptureManagerTemplateTokens.cpp
+++ b/Engine/Plugins/VirtualProduction/CaptureManager/CaptureManagerApp/Source/CaptureManagerSettings/Private/Settings/CaptureManagerTemplateTokens.cpp
@@ -48,89 +48,4 @@ void UCaptureManagerGeneralTokens::OnPostEvaluate_Implementation()
 	Super::OnPostEvaluate_Implementation();
 }
 
-
-UCaptureManagerVideoEncoderTokens::UCaptureManagerVideoEncoderTokens()
-{
-	Namespace = TEXT("cmvidenc");
-
-	VideoEncoderTokens.Emplace(FString(UE::CaptureManager::VideoEncoderTokens::InputKey), { FString(UE::CaptureManager::VideoEncoderTokens::InputKey), LOCTEXT("VideoInputPath", "Input File Path")});
-	VideoEncoderTokens.Emplace(FString(UE::CaptureManager::VideoEncoderTokens::OutputKey), { FString(UE::CaptureManager::VideoEncoderTokens::OutputKey), LOCTEXT("VideoOutputPath", "Output File Path") });
-	VideoEncoderTokens.Emplace(FString(UE::CaptureManager::VideoEncoderTokens::ParamsKey), { FString(UE::CaptureManager::VideoEncoderTokens::ParamsKey), LOCTEXT("VideoParams", "Conversion Parameters (e.g. Pixel Format, Rotation etc") });
-}
-
-UE::CaptureManager::FArchiveToken UCaptureManagerVideoEncoderTokens::GetToken(const FString& InKey) const
-{
-	check(VideoEncoderTokens.Contains(InKey));
-	return VideoEncoderTokens[InKey];
-}
-
-void UCaptureManagerVideoEncoderTokens::OnCreateDefaultTokens(TArray<FNamingTokenData>& OutTokens)
-{
-	Super::OnCreateDefaultTokens(OutTokens);
-
-	for (const TPair<FString, UE::CaptureManager::FArchiveToken>& Token : VideoEncoderTokens)
-	{
-		OutTokens.Add({
-			Token.Value.Name,
-			Token.Value.Description,
-			FNamingTokenData::FTokenProcessorDelegateNative::CreateLambda([Name = Token.Value.Name] {
-					return FText::FromString(Name);
-				})
-			});
-	}
-
-}
-
-void UCaptureManagerVideoEncoderTokens::OnPreEvaluate_Implementation(const FNamingTokensEvaluationData& InEvaluationData)
-{
-	Super::OnPreEvaluate_Implementation(InEvaluationData);
-}
-
-void UCaptureManagerVideoEncoderTokens::OnPostEvaluate_Implementation()
-{
-	Super::OnPostEvaluate_Implementation();
-}
-
-
-UCaptureManagerAudioEncoderTokens::UCaptureManagerAudioEncoderTokens()
-{
-	Namespace = TEXT("cmaudenc");
-
-	AudioEncoderTokens.Emplace(FString(UE::CaptureManager::AudioEncoderTokens::InputKey), { FString(UE::CaptureManager::AudioEncoderTokens::InputKey), LOCTEXT("AudioInputPath", "Input File Path")});
-	AudioEncoderTokens.Emplace(FString(UE::CaptureManager::AudioEncoderTokens::OutputKey), { FString(UE::CaptureManager::AudioEncoderTokens::OutputKey), LOCTEXT("AudioOutputPath", "Output File Path") });
-}
-
-UE::CaptureManager::FArchiveToken UCaptureManagerAudioEncoderTokens::GetToken(const FString& InKey) const
-{
-	check(AudioEncoderTokens.Contains(InKey));
-	return AudioEncoderTokens[InKey];
-}
-
-void UCaptureManagerAudioEncoderTokens::OnCreateDefaultTokens(TArray<FNamingTokenData>& OutTokens)
-{
-	Super::OnCreateDefaultTokens(OutTokens);
-
-	for (const TPair<FString, UE::CaptureManager::FArchiveToken>& Token : AudioEncoderTokens)
-	{
-		OutTokens.Add({
-			Token.Value.Name,
-			Token.Value.Description,
-			FNamingTokenData::FTokenProcessorDelegateNative::CreateLambda([Name = Token.Value.Name] {
-					return FText::FromString(Name);
-				})
-			});
-	}
-
-}
-
-void UCaptureManagerAudioEncoderTokens::OnPreEvaluate_Implementation(const FNamingTokensEvaluationData& InEvaluationData)
-{
-	Super::OnPreEvaluate_Implementation(InEvaluationData);
-}
-
-void UCaptureManagerAudioEncoderTokens::OnPostEvaluate_Implementation()
-{
-	Super::OnPostEvaluate_Implementation();
-}
-
 #undef LOCTEXT_NAMESPACE
diff --git a/Engine/Plugins/VirtualProduction/CaptureManager/CaptureManagerApp/Source/CaptureManagerSettings/Public/Settings/CaptureManagerSettings.h b/Engine/Plugins/VirtualProduction/CaptureManager/CaptureManagerApp/Source/CaptureManagerSettings/Public/Settings/CaptureManagerSettings.h
--- a/Engine/Plugins/VirtualProduction/CaptureManager/CaptureManagerApp/Source/CaptureManagerSettings/Public/Settings/CaptureManagerSettings.h
+++ b/Engine/Plugins/VirtualProduction/CaptureManager/CaptureManagerApp/Source/CaptureManagerSettings/Public/Settings/CaptureManagerSettings.h
@@ -7,6 +7,9 @@
 #include "CaptureManagerSettings.generated.h"
 
 class UCaptureManagerGeneralTokens;
+
+// Encoder token classes live in CaptureDataConverter (Core) so they are auto-discovered
+// by UNamingTokensEngineSubsystem in both LiveLinkHub and UE Editor contexts.
 class UCaptureManagerVideoEncoderTokens;
 class UCaptureManagerAudioEncoderTokens;
 
diff --git a/Engine/Plugins/VirtualProduction/CaptureManager/CaptureManagerApp/Source/CaptureManagerSettings/Public/Settings/CaptureManagerTemplateTokens.h b/Engine/Plugins/VirtualProduction/CaptureManager/CaptureManagerApp/Source/CaptureManagerSettings/Public/Settings/CaptureManagerTemplateTokens.h
--- a/Engine/Plugins/VirtualProduction/CaptureManager/CaptureManagerApp/Source/CaptureManagerSettings/Public/Settings/CaptureManagerTemplateTokens.h
+++ b/Engine/Plugins/VirtualProduction/CaptureManager/CaptureManagerApp/Source/CaptureManagerSettings/Public/Settings/CaptureManagerTemplateTokens.h
@@ -22,18 +22,6 @@ static constexpr FStringView SlateKey = TEXTVIEW("slate");
 static constexpr FStringView TakeKey = TEXTVIEW("take");
 }
 
-namespace VideoEncoderTokens
-{
-static constexpr FStringView InputKey = TEXTVIEW("input");
-static constexpr FStringView OutputKey = TEXTVIEW("output");
-static constexpr FStringView ParamsKey = TEXTVIEW("params");
-}
-
-namespace AudioEncoderTokens
-{
-static constexpr FStringView InputKey = TEXTVIEW("input");
-static constexpr FStringView OutputKey = TEXTVIEW("output");
-}
 }
 
 UCLASS(NotBlueprintable)
@@ -57,46 +45,3 @@ private:
 	TMap<FString, UE::CaptureManager::FArchiveToken> GeneralTokens;
 };
 
-
-UCLASS(NotBlueprintable)
-class CAPTUREMANAGERSETTINGS_API UCaptureManagerVideoEncoderTokens final : public UNamingTokens
-{
-	GENERATED_BODY()
-
-public:
-	UCaptureManagerVideoEncoderTokens();
-
-	UE::CaptureManager::FArchiveToken GetToken(const FString& InKey) const;
-
-protected:
-	// ~Begin UNamingTokens
-	virtual void OnCreateDefaultTokens(TArray<FNamingTokenData>& OutTokens) override;
-	virtual void OnPreEvaluate_Implementation(const FNamingTokensEvaluationData& InEvaluationData) override;
-	virtual void OnPostEvaluate_Implementation() override;
-	// ~End UNamingTokens
-
-private:
-	TMap<FString, UE::CaptureManager::FArchiveToken> VideoEncoderTokens;
-};
-
-
-UCLASS(NotBlueprintable)
-class CAPTUREMANAGERSETTINGS_API UCaptureManagerAudioEncoderTokens final : public UNamingTokens
-{
-	GENERATED_BODY()
-
-public:
-	UCaptureManagerAudioEncoderTokens();
-
-	UE::CaptureManager::FArchiveToken GetToken(const FString& InKey) const;
-
-protected:
-	// ~Begin UNamingTokens
-	virtual void OnCreateDefaultTokens(TArray<FNamingTokenData>& OutTokens) override;
-	virtual void OnPreEvaluate_Implementation(const FNamingTokensEvaluationData& InEvaluationData) override;
-	virtual void OnPostEvaluate_Implementation() override;
-	// ~End UNamingTokens
-
-private:
-	TMap<FString, UE::CaptureManager::FArchiveToken> AudioEncoderTokens;
-};
diff --git a/Engine/Plugins/VirtualProduction/CaptureManager/CaptureManagerApp/Source/IngestLiveLinkDevice/IngestLiveLinkDevice.Build.cs b/Engine/Plugins/VirtualProduction/CaptureManager/CaptureManagerApp/Source/IngestLiveLinkDevice/IngestLiveLinkDevice.Build.cs
--- a/Engine/Plugins/VirtualProduction/CaptureManager/CaptureManagerApp/Source/IngestLiveLinkDevice/IngestLiveLinkDevice.Build.cs
+++ b/Engine/Plugins/VirtualProduction/CaptureManager/CaptureManagerApp/Source/IngestLiveLinkDevice/IngestLiveLinkDevice.Build.cs
@@ -40,7 +40,8 @@ public class IngestLiveLinkDevice : ModuleRules
 			"CoreUObject",
 			"UnrealEd",
 			"ElectraPlayerPlugin",
-			"MJPEGDecoderElectra"
+			"MJPEGDecoderElectra",
+			"NamingTokens"
 		});
 
 		CppCompileWarningSettings.UnsafeTypeCastWarningLevel = WarningLevel.Error;
diff --git a/Engine/Plugins/VirtualProduction/CaptureManager/CaptureManagerApp/Source/IngestLiveLinkDevice/Private/BaseIngestLiveLinkDevice.cpp b/Engine/Plugins/VirtualProduction/CaptureManager/CaptureManagerApp/Source/IngestLiveLinkDevice/Private/BaseIngestLiveLinkDevice.cpp
--- a/Engine/Plugins/VirtualProduction/CaptureManager/CaptureManagerApp/Source/IngestLiveLinkDevice/Private/BaseIngestLiveLinkDevice.cpp
+++ b/Engine/Plugins/VirtualProduction/CaptureManager/CaptureManagerApp/Source/IngestLiveLinkDevice/Private/BaseIngestLiveLinkDevice.cpp
@@ -5,18 +5,76 @@
 #include "HAL/FileManager.h"
 #include "CaptureManagerUnrealEndpointModule.h"
 #include "Settings/CaptureManagerSettings.h"
+#include "Settings/CaptureManagerTemplateTokens.h"
+#include "CaptureManagerEncoderNamingTokens.h"
 
 #include "Utils/IngestLiveLinkDeviceUtils.h"
 #include "Utils/IngestLiveLinkDeviceLog.h"
 
 #include "CaptureDataConverter.h"
+#include "Engine/Engine.h"
 
+#include "NamingTokensEngineSubsystem.h"
 #include "Async/TaskWaiter.h"
 
 DEFINE_LOG_CATEGORY(LogIngestLiveLinkDevice);
 
 namespace UE::CaptureManager::Private
 {
+
+static FString ResolveCommandArgs(const FStringView InDefaultArgs, const FString& InCustomArgs, const FString& InTokenNamespace)
+{
+	FString CommandArgs = InCustomArgs.IsEmpty() ? FString(InDefaultArgs) : InCustomArgs;
+
+	if (!GEngine)
+	{
+		UE_LOGF(LogIngestLiveLinkDevice, Warning, "Cannot resolve encoder command tokens: GEngine is unavailable");
+		return CommandArgs;
+	}
+
+	UNamingTokensEngineSubsystem* NamingTokensSubsystem = GEngine->GetEngineSubsystem<UNamingTokensEngineSubsystem>();
+
+	if (!NamingTokensSubsystem)
+	{
+		UE_LOGF(LogIngestLiveLinkDevice, Warning, "Cannot resolve encoder command tokens: NamingTokensSubsystem is unavailable");
+		return CommandArgs;
+	}
+
+	FNamingTokenFilterArgs FilterArgs;
+	FilterArgs.AdditionalNamespacesToInclude.Add(InTokenNamespace);
+	FilterArgs.bNativeOnly = true;
+
+	CommandArgs = NamingTokensSubsystem->EvaluateTokenString(CommandArgs, FilterArgs).EvaluatedText.ToString();
+
+	return CommandArgs;
+}
+
+static void ApplyThirdPartyEncoderSettings(const UCaptureManagerSettings* InSettings, FCaptureDataConverterParams& OutParams)
+{
+	if (!InSettings || !InSettings->bEnableThirdPartyEncoder || InSettings->ThirdPartyEncoder.FilePath.IsEmpty())
+	{
+		return;
+	}
+
+	const bool bThirdPartyEncoderIsAvailable = FPlatformProcess::ExecProcess(*InSettings->ThirdPartyEncoder.FilePath, TEXT("-version"), nullptr, nullptr, nullptr);
+
+	if (bThirdPartyEncoderIsAvailable)
+	{
+		// {input}, {output} and {params} tokens will remain and are resolved with an FString::Format in the encoder command's Execute() function
+		if (const UCaptureManagerAudioEncoderTokens* AudioTokens = InSettings->GetAudioEncoderNamingTokens())
+		{
+			const FString CommandArgs = ResolveCommandArgs(UE::CaptureManager::EncoderDefaults::AudioCommandArgs, InSettings->CustomAudioCommandArguments, AudioTokens->GetNamespace());
+			OutParams.AudioEncoderConfig = FCaptureManagerEncoderConfig{InSettings->ThirdPartyEncoder.FilePath, CommandArgs};
+		}
+
+		if (const UCaptureManagerVideoEncoderTokens* VideoTokens = InSettings->GetVideoEncoderNamingTokens())
+		{
+			const FString CommandArgs = ResolveCommandArgs(UE::CaptureManager::EncoderDefaults::VideoCommandArgs, InSettings->CustomVideoCommandArguments, VideoTokens->GetNamespace());
+			OutParams.VideoEncoderConfig = FCaptureManagerEncoderConfig{InSettings->ThirdPartyEncoder.FilePath, CommandArgs};
+		}
+	}
+}
+
 static void TryCleanUpWorkingDirectory(const FString& InWorkingDirectory)
 {
 	TRACE_CPUPROFILER_EVENT_SCOPE(FBaseIngestDevice_CleanUpWorkingDirectory);
@@ -409,6 +467,11 @@ FCaptureDataConverterResult<void> UBaseIngestLiveLinkDevice::FImpl::RunConversio
 
 	InParams.CalibrationOutputParams = MoveTemp(CalibrationParams);
 
+	if (const UCaptureManagerSettings* CaptureManagerSettings = GetDefault<UCaptureManagerSettings>())
+	{
+		Private::ApplyThirdPartyEncoderSettings(CaptureManagerSettings, InParams);
+	}
+
 	FCaptureDataConverter::FProgressReporter ProgressReporter =
 		FCaptureDataConverter::FProgressReporter::CreateLambda([ConvertTask](const double InProgress) mutable
 	{
diff --git a/Engine/Plugins/VirtualProduction/CaptureManager/CaptureManagerCore/CaptureManagerCore.uplugin b/Engine/Plugins/VirtualProduction/CaptureManager/CaptureManagerCore/CaptureManagerCore.uplugin
--- a/Engine/Plugins/VirtualProduction/CaptureManager/CaptureManagerCore/CaptureManagerCore.uplugin
+++ b/Engine/Plugins/VirtualProduction/CaptureManager/CaptureManagerCore/CaptureManagerCore.uplugin
@@ -27,6 +27,21 @@
 			"Type": "Runtime",
 			"LoadingPhase": "Default"
 		},
+		{
+			"Name": "CaptureManagerPipeline",
+			"Type": "Editor",
+			"LoadingPhase": "Default"
+		},
+		{
+			"Name": "CaptureManagerMediaRW",
+			"Type": "Editor",
+			"LoadingPhase": "Default"
+		},
+		{
+			"Name": "CaptureDataConverter",
+			"Type": "Editor",
+			"LoadingPhase": "Default"
+		},
 		{
 			"Name": "DataIngestCore",
 			"Type": "Editor",
@@ -52,6 +67,18 @@
 		{
 			"Name": "CaptureData",
 			"Enabled": true
+		},
+		{
+			"Name": "ControlFlows",
+			"Enabled": true
+		},
+		{
+			"Name": "NamingTokens",
+			"Enabled": true
+		},
+		{
+			"Name": "ElectraPlayer",
+			"Enabled": true
 		}
 	]
 }
\ No newline at end of file
