Persisted Media Source gets deleted after reconnecting

#Synced-CL 52095639

diff --git a/Engine/Plugins/VirtualProduction/CaptureManager/CaptureManagerApp/Source/LiveLinkCapabilities/Private/Streaming/LiveLinkDeviceCapability_Streaming.cpp b/Engine/Plugins/VirtualProduction/CaptureManager/CaptureManagerApp/Source/LiveLinkCapabilities/Private/Streaming/LiveLinkDeviceCapability_Streaming.cpp
--- a/Engine/Plugins/VirtualProduction/CaptureManager/CaptureManagerApp/Source/LiveLinkCapabilities/Private/Streaming/LiveLinkDeviceCapability_Streaming.cpp
+++ b/Engine/Plugins/VirtualProduction/CaptureManager/CaptureManagerApp/Source/LiveLinkCapabilities/Private/Streaming/LiveLinkDeviceCapability_Streaming.cpp
@@ -81,6 +81,7 @@ public:
 	void SetMediaSource(UMediaSource* InSource);
 	void RegisterMediaSource(UMediaSource* InMediaSource);
 	void UnregisterMediaSource();
+	bool IsStarted() const { return bIsStarted; }
 
 private:
 
@@ -93,6 +94,7 @@ private:
 
 	mutable FCriticalSection Mutex;
 	TStrongObjectPtr<UMediaSource> CurrentMediaSource = nullptr;
+	std::atomic_bool bIsStarted = false;
 
 	FDelegateHandle DelegateHandle;
 };
@@ -131,9 +133,7 @@ void ILiveLinkDeviceCapability_Streaming::FImpl::RegisterMediaSource(UMediaSourc
 		UMediaProfile* MediaProfile = MediaProfileManager.GetCurrentMediaProfile();
 
 		RegisterMediaSource_GameThread(MediaSource.Get(), MediaProfile);
-
-		SetMediaSource(MediaSource.Get());
-		MediaSource.Reset();
+		bIsStarted = true;
 	}
 	else
 	{
@@ -145,9 +145,7 @@ void ILiveLinkDeviceCapability_Streaming::FImpl::RegisterMediaSource(UMediaSourc
 				UMediaProfile* MediaProfile = MediaProfileManager.GetCurrentMediaProfile();
 
 				SharedThis->RegisterMediaSource_GameThread(MediaSource.Get(), MediaProfile);
-
-				SharedThis->SetMediaSource(MediaSource.Get());
-				MediaSource.Reset();
+				SharedThis->bIsStarted = true;
 			}
 		});
 	}
@@ -161,7 +159,7 @@ void ILiveLinkDeviceCapability_Streaming::FImpl::UnregisterMediaSource()
 		UMediaProfile* MediaProfile = MediaProfileManager.GetCurrentMediaProfile();
 
 		UnregisterMediaSource_GameThread(MediaProfile);
-		SetMediaSource(nullptr);
+		bIsStarted = false;
 	}
 	else
 	{
@@ -173,7 +171,7 @@ void ILiveLinkDeviceCapability_Streaming::FImpl::UnregisterMediaSource()
 				UMediaProfile* MediaProfile = MediaProfileManager.GetCurrentMediaProfile();
 
 				SharedThis->UnregisterMediaSource_GameThread(MediaProfile);
-				SharedThis->SetMediaSource(nullptr);
+				SharedThis->bIsStarted = false;
 			}
 		});
 	}
@@ -188,6 +186,11 @@ void ILiveLinkDeviceCapability_Streaming::FImpl::OnMediaProfileChanged(UMediaPro
 		return;
 	}
 
+	if (!IsStarted())
+	{
+		return;
+	}
+
 	if (InPrevious)
 	{
 		UnregisterMediaSource_GameThread(InPrevious);
@@ -285,29 +288,34 @@ ILiveLinkDeviceCapability_Streaming::ILiveLinkDeviceCapability_Streaming()
 {
 }
 
-void ILiveLinkDeviceCapability_Streaming::Start_Implementation()
+void ILiveLinkDeviceCapability_Streaming::StartStreaming_Implementation()
 {
 	Impl->RegisterMediaSource(Impl->GetMediaSource());
 }
 
-void ILiveLinkDeviceCapability_Streaming::Stop_Implementation()
+void ILiveLinkDeviceCapability_Streaming::StopStreaming_Implementation()
 {
 	Impl->UnregisterMediaSource();
 }
 
+bool ILiveLinkDeviceCapability_Streaming::IsStreaming_Implementation() const
+{
+	return Impl->IsStarted();
+}
+
 UMediaSource* ILiveLinkDeviceCapability_Streaming::GetMediaSource_Implementation() const
 {
 	return Impl->GetMediaSource();
 }
 
-void ILiveLinkDeviceCapability_Streaming::RegisterMediaSource(UMediaSource* InMediaSource)
+void ILiveLinkDeviceCapability_Streaming::SetMediaSource(UMediaSource* InMediaSource)
 {
-	Impl->RegisterMediaSource(InMediaSource);
+	Impl->SetMediaSource(InMediaSource);
 }
 
-void ILiveLinkDeviceCapability_Streaming::UnregisterMediaSource()
+void ILiveLinkDeviceCapability_Streaming::DeleteMediaSource()
 {
-	Impl->UnregisterMediaSource();
+	Impl->SetMediaSource(nullptr);
 }
 
 void ILiveLinkDeviceCapability_Streaming::Init()
diff --git a/Engine/Plugins/VirtualProduction/CaptureManager/CaptureManagerApp/Source/LiveLinkCapabilities/Public/Streaming/LiveLinkDeviceCapability_Streaming.h b/Engine/Plugins/VirtualProduction/CaptureManager/CaptureManagerApp/Source/LiveLinkCapabilities/Public/Streaming/LiveLinkDeviceCapability_Streaming.h
--- a/Engine/Plugins/VirtualProduction/CaptureManager/CaptureManagerApp/Source/LiveLinkCapabilities/Public/Streaming/LiveLinkDeviceCapability_Streaming.h
+++ b/Engine/Plugins/VirtualProduction/CaptureManager/CaptureManagerApp/Source/LiveLinkCapabilities/Public/Streaming/LiveLinkDeviceCapability_Streaming.h
@@ -43,12 +43,16 @@ public:
 	ILiveLinkDeviceCapability_Streaming();
 
 	UFUNCTION(BlueprintCallable, BlueprintNativeEvent, Category = "Live Link Device|Streaming")
-	void Start();
-	void Start_Implementation();
+	void StartStreaming();
+	void StartStreaming_Implementation();
 
 	UFUNCTION(BlueprintCallable, BlueprintNativeEvent, Category = "Live Link Device|Streaming")
-	void Stop();
-	void Stop_Implementation();
+	void StopStreaming();
+	void StopStreaming_Implementation();
+
+	UFUNCTION(BlueprintCallable, BlueprintNativeEvent, Category = "Live Link Device|Streaming")
+	bool IsStreaming() const;
+	bool IsStreaming_Implementation() const;
 
 	UFUNCTION(BlueprintCallable, BlueprintNativeEvent, Category = "Live Link Device|Streaming")
 	class UMediaSource* GetMediaSource() const;
@@ -56,8 +60,8 @@ public:
 
 protected:
 
-	void RegisterMediaSource(class UMediaSource* InMediaSource);
-	void UnregisterMediaSource();
+	void SetMediaSource(class UMediaSource* InMediaSource);
+	void DeleteMediaSource();
 
 private:
 
diff --git a/Engine/Plugins/VirtualProduction/CaptureManager/CaptureManagerDevices/Source/CPSLiveLinkDevice/Private/LiveLinkFaceDevice.cpp b/Engine/Plugins/VirtualProduction/CaptureManager/CaptureManagerDevices/Source/CPSLiveLinkDevice/Private/LiveLinkFaceDevice.cpp
--- a/Engine/Plugins/VirtualProduction/CaptureManager/CaptureManagerDevices/Source/CPSLiveLinkDevice/Private/LiveLinkFaceDevice.cpp
+++ b/Engine/Plugins/VirtualProduction/CaptureManager/CaptureManagerDevices/Source/CPSLiveLinkDevice/Private/LiveLinkFaceDevice.cpp
@@ -86,7 +86,8 @@ void ULiveLinkFaceDevice::OnLoadDeviceCapabilityData(TMap<TSubclassOf<ULiveLinkD
 		{
 			if (TObjectPtr<URtspMediaSource> MediaSource = Cast<URtspMediaSource>(CapabilityData.Value))
 			{
-				RegisterMediaSource(MediaSource);
+				SetMediaSource(MediaSource);
+				StartStreaming_Implementation();
 			}
 		}
 	}
@@ -361,14 +362,17 @@ void ULiveLinkFaceDevice::HandleConnectionChanged(TSharedPtr<const UE::CaptureMa
 			RtspMediaSource->DecoderBufferSize = 1; // Using 1 for minimum latency
 			RtspMediaSource->bProvideCpuBuffer = true;
 
-			RegisterMediaSource(RtspMediaSource);
+			SetMediaSource(RtspMediaSource);
 		}
 
+		// Start streaming as soon as device is connected
+		StartStreaming_Implementation();
+
 		bIsConnecting = false;
 	}
 	else if (Event->ConnectionState == FConnectionStateChangedEvent::EState::Disconnected)
 	{
-		UnregisterMediaSource();
+		StopStreaming_Implementation();
 	}
 
 	auto GetConnectionState = [](FConnectionStateChangedEvent::EState InConnectionState) -> ELiveLinkDeviceConnectionStatus
