[Backout] - CL52058844
#rnx
Original CL Desc
-----------------------------------------------------------------
[CaptureManager] Moved conversion modules to CaptureManagerCore.

Moves CaptureManagerPipeline, CaptureManagerMediaRW, and CaptureDataConverter
from CaptureManagerApp (Editor type) to CaptureManagerCore (Editor type),
enabling both LiveLinkHub and UE Editor workflows to share identical conversion
logic from a single compiled module.

#Synced-CL 52058844

diff --git a/Engine/Plugins/VirtualProduction/CaptureManager/CaptureManagerApp/CaptureManagerApp.uplugin b/Engine/Plugins/VirtualProduction/CaptureManager/CaptureManagerApp/CaptureManagerApp.uplugin
--- a/Engine/Plugins/VirtualProduction/CaptureManager/CaptureManagerApp/CaptureManagerApp.uplugin
+++ b/Engine/Plugins/VirtualProduction/CaptureManager/CaptureManagerApp/CaptureManagerApp.uplugin
@@ -33,6 +33,21 @@
 			"Type": "Editor",
 			"LoadingPhase": "Default"
 		},
+		{
+			"Name": "CaptureManagerPipeline",
+			"Type": "Editor",
+			"LoadingPhase": "Default"
+		},
+		{
+			"Name": "CaptureManagerMediaRW",
+			"Type": "Editor",
+			"LoadingPhase": "Default"
+		},
+		{
+			"Name": "CaptureDataConverter",
+			"Type": "Editor",
+			"LoadingPhase": "Default"
+		},
 		{
 			"Name": "CaptureManagerUnrealEndpoint",
 			"Type": "Editor",
@@ -59,6 +74,14 @@
 			"Name": "LiveLinkHub",
 			"Enabled": true
 		},
+		{
+			"Name": "NamingTokens",
+			"Enabled": true
+		},
+		{
+			"Name": "ControlFlows",
+			"Enabled": true
+		},
 		{
 			"Name": "CaptureManagerCore",
 			"Enabled": true
diff --git a/Engine/Plugins/VirtualProduction/CaptureManager/CaptureManagerApp/Source/CaptureDataConverter/Private/CaptureDataConverterModule.cpp b/Engine/Plugins/VirtualProduction/CaptureManager/CaptureManagerApp/Source/CaptureDataConverter/Private/CaptureDataConverterModule.cpp
new file mode 100644
--- /dev/null
+++ b/Engine/Plugins/VirtualProduction/CaptureManager/CaptureManagerApp/Source/CaptureDataConverter/Private/CaptureDataConverterModule.cpp
@@ -0,0 +1,93 @@
+// Copyright Epic Games, Inc. All Rights Reserved.
+
+#include "CaptureDataConverterModule.h"
+
+#include "HAL/PlatformProcess.h"
+
+#include "Settings/CaptureManagerSettings.h"
+
+#include "Nodes/ThirdPartyEncoder/CaptureThirdPartyNodeParams.h"
+
+namespace UE::CaptureManager::Private
+{
+
+bool ThirdPartyAvailabilityCheck(const FString& InThirdPartyEncoder)
+{
+	return FPlatformProcess::ExecProcess(*InThirdPartyEncoder, TEXT("-version"), nullptr, nullptr, nullptr);
+}
+
+}
+
+bool FCaptureDataConverterModule::IsThirdPartyEncoderAvailable()
+{
+	return CheckThirdPartyEncoderAvailability();
+}
+
+FString FCaptureDataConverterModule::GetThirdPartyEncoder() const
+{
+	const UCaptureManagerSettings* Settings = GetDefault<UCaptureManagerSettings>();
+
+	if (!Settings->bEnableThirdPartyEncoder)
+	{
+		return FString();
+	}
+
+	return Settings->ThirdPartyEncoder.FilePath;
+}
+
+FString FCaptureDataConverterModule::GetThirdPartyEncoderVideoCommandArguments() const
+{
+	const UCaptureManagerSettings* Settings = GetDefault<UCaptureManagerSettings>();
+
+	if (!Settings->bEnableThirdPartyEncoder)
+	{
+		return FString();
+	}
+
+	return Settings->CustomVideoCommandArguments;
+}
+
+FString FCaptureDataConverterModule::GetThirdPartyEncoderAudioCommandArguments() const
+{
+	const UCaptureManagerSettings* Settings = GetDefault<UCaptureManagerSettings>();
+
+	if (!Settings->bEnableThirdPartyEncoder)
+	{
+		return FString();
+	}
+
+	return Settings->CustomAudioCommandArguments;
+}
+
+void FCaptureDataConverterModule::StartupModule()
+{
+	UCaptureManagerSettings* Settings = GetMutableDefault<UCaptureManagerSettings>();
+
+	if (Settings->CustomAudioCommandArguments.IsEmpty())
+	{
+		Settings->CustomVideoCommandArguments = UE::CaptureManager::VideoCommandArgumentTemplate;
+	}
+	
+	if (Settings->CustomAudioCommandArguments.IsEmpty())
+	{
+		Settings->CustomAudioCommandArguments = UE::CaptureManager::AudioCommandArgumentTemplate;
+	}
+}
+
+void FCaptureDataConverterModule::ShutdownModule()
+{
+}
+
+bool FCaptureDataConverterModule::CheckThirdPartyEncoderAvailability() const
+{
+	const UCaptureManagerSettings* Settings = GetDefault<UCaptureManagerSettings>();
+
+	if (Settings->bEnableThirdPartyEncoder && !Settings->ThirdPartyEncoder.FilePath.IsEmpty())
+	{
+		return UE::CaptureManager::Private::ThirdPartyAvailabilityCheck(Settings->ThirdPartyEncoder.FilePath);
+	}
+
+	return false;
+}
+
+IMPLEMENT_MODULE(FCaptureDataConverterModule, CaptureDataConverter)
\ No newline at end of file
diff --git a/Engine/Plugins/VirtualProduction/CaptureManager/CaptureManagerApp/Source/CaptureDataConverter/Private/CaptureDataConverterModule.h b/Engine/Plugins/VirtualProduction/CaptureManager/CaptureManagerApp/Source/CaptureDataConverter/Private/CaptureDataConverterModule.h
new file mode 100644
--- /dev/null
+++ b/Engine/Plugins/VirtualProduction/CaptureManager/CaptureManagerApp/Source/CaptureDataConverter/Private/CaptureDataConverterModule.h
@@ -0,0 +1,24 @@
+// Copyright Epic Games, Inc. All Rights Reserved.
+
+#pragma once
+
+#include "Modules/ModuleManager.h"
+
+class FCaptureDataConverterModule : public IModuleInterface
+{
+public:
+	
+	bool IsThirdPartyEncoderAvailable();
+	FString GetThirdPartyEncoder() const;
+	FString GetThirdPartyEncoderVideoCommandArguments() const;
+	FString GetThirdPartyEncoderAudioCommandArguments() const;
+
+private:
+
+	virtual void StartupModule() override;
+	virtual void ShutdownModule() override;
+
+	bool CheckThirdPartyEncoderAvailability() const;
+
+	bool bIsThirdPartyEncoderEnabled = false;
+};
\ No newline at end of file
diff --git a/Engine/Plugins/VirtualProduction/CaptureManager/CaptureManagerApp/Source/CaptureDataConverter/Private/Nodes/ThirdPartyEncoder/CaptureConvertAudioDataThirdParty.cpp b/Engine/Plugins/VirtualProduction/CaptureManager/CaptureManagerApp/Source/CaptureDataConverter/Private/Nodes/ThirdPartyEncoder/CaptureConvertAudioDataThirdParty.cpp
new file mode 100644
--- /dev/null
+++ b/Engine/Plugins/VirtualProduction/CaptureManager/CaptureManagerApp/Source/CaptureDataConverter/Private/Nodes/ThirdPartyEncoder/CaptureConvertAudioDataThirdParty.cpp
@@ -0,0 +1,157 @@
+// Copyright Epic Games, Inc. All Rights Reserved.
+
+#include "CaptureConvertAudioDataThirdParty.h"
+
+#include "MediaSample.h"
+
+#include "Async/HelperFunctions.h"
+#include "Async/StopToken.h"
+#include "Async/TaskProgress.h"
+
+#include "Nodes/CaptureCopyProgressReporter.h"
+#include "CaptureThirdPartyNodeUtils.h"
+
+#include "CaptureManagerMediaRWModule.h"
+
+#include "Engine/Engine.h"
+#include "GlobalNamingTokens.h"
+#include "NamingTokenData.h"
+#include "NamingTokensEngineSubsystem.h"
+#include "Settings/CaptureManagerSettings.h"
+#include "Settings/CaptureManagerTemplateTokens.h"
+
+#include "ProcessRunner/ProcessRunner.h"
+
+#define LOCTEXT_NAMESPACE "CaptureConvertAudioDataTP"
+
+DEFINE_LOG_CATEGORY_STATIC(LogCaptureConvertAudioDataThirdParty, Log, All);
+
+FCaptureConvertAudioDataThirdParty::FCaptureConvertAudioDataThirdParty(FCaptureThirdPartyNodeParams InThirdPartyEncoder,
+																	   const FTakeMetadata::FAudio& InAudio,
+																	   const FString& InOutputDirectory,
+																	   const FCaptureConvertDataNodeParams& InParams,
+																	   const FCaptureConvertAudioOutputParams& InAudioParams)
+	: FConvertAudioNode(InAudio, InOutputDirectory)
+	, ThirdPartyEncoder(MoveTemp(InThirdPartyEncoder))
+	, Params(InParams)
+	, AudioParams(InAudioParams)
+{
+}
+
+FCaptureConvertAudioDataThirdParty::FResult FCaptureConvertAudioDataThirdParty::Run()
+{
+	using namespace UE::CaptureManager;
+
+	if (Params.StopToken.IsStopRequested())
+	{
+		FText Message = LOCTEXT("CaptureConvertAudioNodeTP_AbortedByUser", "Aborted by user");
+		return MakeError(MoveTemp(Message));
+	}
+
+	if (FPaths::GetExtension(Audio.Path) == AudioParams.Format)
+	{
+		return CopyAudioFile();
+	}
+
+	return ConvertAudioFile();
+}
+
+FCaptureConvertAudioDataThirdParty::FResult FCaptureConvertAudioDataThirdParty::CopyAudioFile()
+{
+	using namespace UE::CaptureManager;
+
+	FTaskProgress::FTask Task = Params.TaskProgress->StartTask();
+
+	FString DestinationDirectory = OutputDirectory / Audio.Name;
+	const FString AudioFilePath = FPaths::ConvertRelativePathToFull(Params.TakeOriginDirectory, Audio.Path);
+
+	FCopyProgressReporter ProgressReporter(Task, Params.StopToken);
+	FString Destination = DestinationDirectory / FPaths::SetExtension(AudioParams.AudioFileName, AudioParams.Format);
+
+	constexpr bool bReplace = true;
+	constexpr bool bEvenIfReadOnly = true;
+	constexpr bool bAttributes = false;
+	uint32 Result = IFileManager::Get().Copy(*Destination, *AudioFilePath, bReplace, bEvenIfReadOnly, bAttributes, &ProgressReporter);
+
+	if (Result == COPY_Fail)
+	{
+		FText Message = FText::Format(LOCTEXT("CaptureConvertAudioNodeTP_CopyFailed", "Failed to copy the audio file {0}"), 
+									  FText::FromString(AudioFilePath));
+		return MakeError(MoveTemp(Message));
+	}
+
+	if (Result == COPY_Canceled)
+	{
+		FText Message = LOCTEXT("CaptureConvertAudioNodeTP_AbortedByUser", "Aborted by user");
+		return MakeError(MoveTemp(Message));
+	}
+
+	return MakeValue();
+}
+
+FCaptureConvertAudioDataThirdParty::FResult FCaptureConvertAudioDataThirdParty::ConvertAudioFile()
+{
+	using namespace UE::CaptureManager;
+
+	FTaskProgress::FTask Task = Params.TaskProgress->StartTask();
+
+	FString DestinationDirectory = OutputDirectory / Audio.Name;
+	const FString AudioFilePath = FPaths::ConvertRelativePathToFull(Params.TakeOriginDirectory, Audio.Path);
+
+	const FString AudioOutputFile = FPaths::SetExtension(FPaths::Combine(DestinationDirectory, AudioParams.AudioFileName), AudioParams.Format);
+
+	if (ThirdPartyEncoder.CommandArguments.IsEmpty())
+	{
+		ThirdPartyEncoder.CommandArguments = UE::CaptureManager::AudioCommandArgumentTemplate;
+	}
+
+	FString CommandArgs = ThirdPartyEncoder.CommandArguments;
+
+	using namespace UE::CaptureManager;
+
+	UNamingTokensEngineSubsystem* NamingTokensSubsystem = GEngine->GetEngineSubsystem<UNamingTokensEngineSubsystem>();
+	check(NamingTokensSubsystem);
+
+	const UCaptureManagerSettings* Settings = GetDefault<UCaptureManagerSettings>();
+
+	FNamingTokenFilterArgs AudioEncoderTokenArgs;
+	const TObjectPtr<const UCaptureManagerAudioEncoderTokens> Tokens = Settings->GetAudioEncoderNamingTokens();
+	check(Tokens);
+	AudioEncoderTokenArgs.AdditionalNamespacesToInclude.Add(Tokens->GetNamespace());
+	AudioEncoderTokenArgs.bNativeOnly = true;
+
+	FStringFormatNamedArguments AudioEncoderFormatArgs;
+	AudioEncoderFormatArgs.Add(Tokens->GetToken(FString(AudioEncoderTokens::InputKey)).Name, WrapInQuotes(AudioFilePath));
+	AudioEncoderFormatArgs.Add(Tokens->GetToken(FString(AudioEncoderTokens::OutputKey)).Name, WrapInQuotes(AudioOutputFile));
+
+	CommandArgs = FString::Format(*CommandArgs, AudioEncoderFormatArgs);
+	FNamingTokenResultData AudioEncoderCommandResult = NamingTokensSubsystem->EvaluateTokenString(CommandArgs, AudioEncoderTokenArgs);
+	CommandArgs = AudioEncoderCommandResult.EvaluatedText.ToString();
+
+	FProcessRunnerResult RunnerResult = FProcessRunner::Run(ThirdPartyEncoder.Encoder, CommandArgs, &Params.StopToken);
+
+	if (RunnerResult.HasError())
+	{
+		EProcessRunnerError RunnerError = RunnerResult.StealError();
+
+		FText Message;
+		switch (RunnerError)
+		{
+			case EProcessRunnerError::LaunchFailed:
+				Message = LOCTEXT("CaptureConvertAudioNodeTP_ProcessNotFound", "Failed to start the process");
+				break;
+			case EProcessRunnerError::Cancelled:
+				Message = LOCTEXT("CaptureConvertAudioNodeTP_AbortedByUser", "Aborted by user");
+				break;
+			default:
+				Message = LOCTEXT("CaptureConvertAudioNodeTP_ErrorRunning", "Error occurred while running the third party encoder");
+				break;
+		}
+		return MakeError(MoveTemp(Message));
+	}
+
+	Task.Update(1.0f);
+	return MakeValue();
+}
+
+#undef LOCTEXT_NAMESPACE
\ No newline at end of file
diff --git a/Engine/Plugins/VirtualProduction/CaptureManager/CaptureManagerApp/Source/CaptureDataConverter/Private/Nodes/ThirdPartyEncoder/CaptureConvertAudioDataThirdParty.h b/Engine/Plugins/VirtualProduction/CaptureManager/CaptureManagerApp/Source/CaptureDataConverter/Private/Nodes/ThirdPartyEncoder/CaptureConvertAudioDataThirdParty.h
new file mode 100644
--- /dev/null
+++ b/Engine/Plugins/VirtualProduction/CaptureManager/CaptureManagerApp/Source/CaptureDataConverter/Private/Nodes/ThirdPartyEncoder/CaptureConvertAudioDataThirdParty.h
@@ -0,0 +1,32 @@
+// Copyright Epic Games, Inc. All Rights Reserved.
+
+#pragma once
+
+#include "Nodes/ConvertAudioNode.h"
+
+#include "CaptureThirdPartyNodeParams.h"
+
+#include "CaptureDataConverterNodeParams.h"
+
+class FCaptureConvertAudioDataThirdParty final :
+	public FConvertAudioNode
+{
+public:
+
+	FCaptureConvertAudioDataThirdParty(FCaptureThirdPartyNodeParams InThirdPartEncoder,
+									   const FTakeMetadata::FAudio& InAudio,
+									   const FString& InOutputDirectory,
+									   const FCaptureConvertDataNodeParams& InParams,
+									   const FCaptureConvertAudioOutputParams& InAudioParams);
+
+private:
+
+	virtual FResult Run() override;
+
+	FResult CopyAudioFile();
+	FResult ConvertAudioFile();
+
+	FCaptureThirdPartyNodeParams ThirdPartyEncoder;
+	FCaptureConvertDataNodeParams Params;
+	FCaptureConvertAudioOutputParams AudioParams;
+};
\ No newline at end of file
diff --git a/Engine/Plugins/VirtualProduction/CaptureManager/CaptureManagerApp/Source/CaptureDataConverter/Private/Nodes/ThirdPartyEncoder/CaptureThirdPartyNodeParams.h b/Engine/Plugins/VirtualProduction/CaptureManager/CaptureManagerApp/Source/CaptureDataConverter/Private/Nodes/ThirdPartyEncoder/CaptureThirdPartyNodeParams.h
new file mode 100644
--- /dev/null
+++ b/Engine/Plugins/VirtualProduction/CaptureManager/CaptureManagerApp/Source/CaptureDataConverter/Private/Nodes/ThirdPartyEncoder/CaptureThirdPartyNodeParams.h
@@ -0,0 +1,19 @@
+// Copyright Epic Games, Inc. All Rights Reserved.
+
+#pragma once
+
+#include "Containers/UnrealString.h"
+
+namespace UE::CaptureManager
+{
+
+static const FString VideoCommandArgumentTemplate = TEXT("-noautorotate -i {input} -an {params} -q:v 1 -qmax 1 -qmin 1 -start_number 0 {output}");
+static const FString AudioCommandArgumentTemplate = TEXT("-i {input} -vn {output}");
+
+}
+
+struct FCaptureThirdPartyNodeParams
+{
+	FString Encoder;
+	FString CommandArguments;
+};
diff --git a/Engine/Plugins/VirtualProduction/CaptureManager/CaptureManagerApp/Source/CaptureManagerMediaRW/Public/CaptureManagerMediaRWModule.h b/Engine/Plugins/VirtualProduction/CaptureManager/CaptureManagerApp/Source/CaptureManagerMediaRW/Public/CaptureManagerMediaRWModule.h
new file mode 100644
--- /dev/null
+++ b/Engine/Plugins/VirtualProduction/CaptureManager/CaptureManagerApp/Source/CaptureManagerMediaRW/Public/CaptureManagerMediaRWModule.h
@@ -0,0 +1,21 @@
+// Copyright Epic Games, Inc. All Rights Reserved.
+
+#pragma once
+
+#include "Modules/ModuleManager.h"
+
+#include "MediaRWManager.h"
+
+class CAPTUREMANAGERMEDIARW_API FCaptureManagerMediaRWModule : public IModuleInterface
+{
+public:
+
+	void StartupModule();
+	void ShutdownModule();
+
+	FMediaRWManager& Get();
+
+private:
+
+	TUniquePtr<FMediaRWManager> MediaRWManager;
+};
\ No newline at end of file
diff --git a/Engine/Plugins/VirtualProduction/CaptureManager/CaptureManagerApp/Source/CaptureManagerMediaRW/Public/IMediaRWFactory.h b/Engine/Plugins/VirtualProduction/CaptureManager/CaptureManagerApp/Source/CaptureManagerMediaRW/Public/IMediaRWFactory.h
new file mode 100644
--- /dev/null
+++ b/Engine/Plugins/VirtualProduction/CaptureManager/CaptureManagerApp/Source/CaptureManagerMediaRW/Public/IMediaRWFactory.h
@@ -0,0 +1,60 @@
+// Copyright Epic Games, Inc. All Rights Reserved.
+
+#pragma once
+
+#include "IMediaReader.h"
+#include "IMediaWriter.h"
+
+class CAPTUREMANAGERMEDIARW_API IAudioReaderFactory
+{
+public:
+
+	virtual ~IAudioReaderFactory() = default;
+
+	virtual TUniquePtr<IAudioReader> CreateAudioReader() = 0;
+};
+
+class CAPTUREMANAGERMEDIARW_API IVideoReaderFactory
+{
+public:
+
+	virtual ~IVideoReaderFactory() = default;
+
+	virtual TUniquePtr<IVideoReader> CreateVideoReader() = 0;
+};
+
+class CAPTUREMANAGERMEDIARW_API ICalibrationReaderFactory
+{
+public:
+
+	virtual ~ICalibrationReaderFactory() = default;
+
+	virtual TUniquePtr<ICalibrationReader> CreateCalibrationReader() = 0;
+};
+
+class CAPTUREMANAGERMEDIARW_API IAudioWriterFactory
+{
+public:
+
+	virtual ~IAudioWriterFactory() = default;
+
+	virtual TUniquePtr<IAudioWriter> CreateAudioWriter() = 0;
+};
+
+class CAPTUREMANAGERMEDIARW_API IImageWriterFactory
+{
+public:
+
+	virtual ~IImageWriterFactory() = default;
+
+	virtual TUniquePtr<IImageWriter> CreateImageWriter() = 0;
+};
+
+class CAPTUREMANAGERMEDIARW_API ICalibrationWriterFactory
+{
+public:
+
+	virtual ~ICalibrationWriterFactory() = default;
+
+	virtual TUniquePtr<ICalibrationWriter> CreateCalibrationWriter() = 0;
+};
\ No newline at end of file
diff --git a/Engine/Plugins/VirtualProduction/CaptureManager/CaptureManagerApp/Source/CaptureManagerMediaRW/Public/MediaRWManager.h b/Engine/Plugins/VirtualProduction/CaptureManager/CaptureManagerApp/Source/CaptureManagerMediaRW/Public/MediaRWManager.h
new file mode 100644
--- /dev/null
+++ b/Engine/Plugins/VirtualProduction/CaptureManager/CaptureManagerApp/Source/CaptureManagerMediaRW/Public/MediaRWManager.h
@@ -0,0 +1,61 @@
+// Copyright Epic Games, Inc. All Rights Reserved.
+
+#pragma once
+
+#include "IMediaRWFactory.h"
+
+#include "Containers/Map.h"
+
+class CAPTUREMANAGERMEDIARW_API FMediaRWManager
+{
+public:
+
+	FMediaRWManager();
+
+	FMediaRWManager(const FMediaRWManager& InCopy) = delete;
+	FMediaRWManager& operator=(const FMediaRWManager& InCopy) = delete;
+
+	void RegisterAudioReader(const TArray<FString>& InFormats, TUniquePtr<IAudioReaderFactory> InReader);
+	void RegisterVideoReader(const TArray<FString>& InFormats, TUniquePtr<IVideoReaderFactory> InReader);
+	void RegisterCalibrationReader(const TArray<FString>& InFormats, TUniquePtr<ICalibrationReaderFactory> InReader);
+
+	void RegisterAudioWriter(const TArray<FString>& InFormats, TUniquePtr<IAudioWriterFactory> InWriter);
+	void RegisterImageWriter(const TArray<FString>& InFormats, TUniquePtr<IImageWriterFactory> InWriter);
+	void RegisterCalibrationWriter(const TArray<FString>& InFormats, TUniquePtr<ICalibrationWriterFactory> InWriter);
+
+	TUniquePtr<IAudioReader> CreateAudioReaderByFormat(const FString& InFormat, int32 InIndex = 0);
+	TUniquePtr<IVideoReader> CreateVideoReaderByFormat(const FString& InFormat, int32 InIndex = 0);
+	TUniquePtr<ICalibrationReader> CreateCalibrationReaderByFormat(const FString& InFormat, int32 InIndex = 0);
+
+	TValueOrError<TUniquePtr<IAudioReader>, FText> CreateAudioReader(const FString& InPath, int32 InIndex = 0);
+	TValueOrError<TUniquePtr<IVideoReader>, FText> CreateVideoReader(const FString& InPath, int32 InIndex = 0);
+	TValueOrError<TUniquePtr<ICalibrationReader>, FText> CreateCalibrationReader(const FString& InPath, int32 InIndex = 0);
+
+	TUniquePtr<IAudioWriter> CreateAudioWriterByFormat(const FString& InFormat, int32 InIndex = 0);
+	TUniquePtr<IImageWriter> CreateImageWriterByFormat(const FString& InFormat, int32 InIndex = 0);
+	TUniquePtr<ICalibrationWriter> CreateCalibrationWriterByFormat(const FString& InFormat, int32 InIndex = 0);
+
+	TValueOrError<TUniquePtr<IAudioWriter>, FText> CreateAudioWriter(const FString& InDirectory, const FString& InFileName, const FString& InFormat, int32 InIndex = 0);
+	TValueOrError<TUniquePtr<IImageWriter>, FText> CreateImageWriter(const FString& InDirectory, const FString& InFileName, const FString& InFormat, int32 InIndex = 0);
+	TValueOrError<TUniquePtr<ICalibrationWriter>, FText> CreateCalibrationWriter(const FString& InDirectory, const FString& InFileName, const FString& InFormat, int32 InIndex = 0);
+
+private:
+
+	TMap<FString, IAudioReaderFactory*> AudioReadersPerFormat;
+	TArray<TUniquePtr<IAudioReaderFactory>> AudioReaders;
+
+	TMap<FString, IVideoReaderFactory*> VideoReadersPerFormat;
+	TArray<TUniquePtr<IVideoReaderFactory>> VideoReaders;
+
+	TMap<FString, ICalibrationReaderFactory*> CalibrationReadersPerFormat;
+	TArray<TUniquePtr<ICalibrationReaderFactory>> CalibrationReaders;
+
+	TMap<FString, IAudioWriterFactory*> AudioWritersPerFormat;
+	TArray<TUniquePtr<IAudioWriterFactory>> AudioWriters;
+
+	TMap<FString, IImageWriterFactory*> ImageWritersPerFormat;
+	TArray<TUniquePtr<IImageWriterFactory>> ImageWriters;
+
+	TMap<FString, ICalibrationWriterFactory*> CalibrationWritersPerFormat;
+	TArray<TUniquePtr<ICalibrationWriterFactory>> CalibrationWriters;
+};
\ No newline at end of file
diff --git a/Engine/Plugins/VirtualProduction/CaptureManager/CaptureManagerApp/Source/CaptureManagerMediaRW/Public/MediaSample.h b/Engine/Plugins/VirtualProduction/CaptureManager/CaptureManagerApp/Source/CaptureManagerMediaRW/Public/MediaSample.h
new file mode 100644
--- /dev/null
+++ b/Engine/Plugins/VirtualProduction/CaptureManager/CaptureManagerApp/Source/CaptureManagerMediaRW/Public/MediaSample.h
@@ -0,0 +1,166 @@
+// Copyright Epic Games, Inc. All Rights Reserved.
+
+#pragma once
+
+#include "IMediaAudioSample.h"
+#include "IMediaTextureSample.h"
+
+namespace UE::CaptureManager
+{
+
+enum class ESampleRate
+{
+	SR_8000Hz = 0,
+	SR_16000Hz,
+	SR_44100Hz,
+	SR_48000Hz,
+	SR_88200Hz,
+	SR_96000Hz,
+	SR_192000Hz
+};
+
+struct FMediaAudioSample
+{
+	TArray<uint8> Buffer;
+	uint32 Channels;
+	FTimespan Duration;
+	EMediaAudioSampleFormat SampleFormat;
+	uint32 Frames;
+	ESampleRate SampleRate;
+	FTimespan Time;
+};
+
+CAPTUREMANAGERMEDIARW_API int32 ConvertBitsPerSample(EMediaAudioSampleFormat InSampleFormat);
+CAPTUREMANAGERMEDIARW_API int32 ConvertSampleRate(ESampleRate InSampleRate);
+
+CAPTUREMANAGERMEDIARW_API ESampleRate ConvertSampleRate(int32 InSampleRate);
+
+enum class EMediaTexturePixelFormat
+{
+	Undefined = -1,
+	U8_RGB = 0,
+	U8_BGR,
+	U8_RGBA,
+	U8_BGRA,
+	U8_I444,
+	U8_I420,
+	U8_YUY2,
+	U8_NV12,
+	U8_Mono,
+	U16_Mono,
+	F_Mono,
+
+	Default = U8_RGB
+};
+
+struct FMediaTextureSample
+{
+	TArray<uint8> Buffer;
+	int32 Stride = -1;
+	FIntPoint Dimensions;
+	FTimespan Duration;
+	EMediaTexturePixelFormat CurrentFormat;
+	EMediaTexturePixelFormat DesiredFormat;
+	FTimespan Time;
+	EMediaOrientation Orientation;
+	EMediaOrientation Rotation;
+};
+
+CAPTUREMANAGERMEDIARW_API uint32 GetNumberOfChannels(EMediaTexturePixelFormat InPixelFormat);
+
+struct FOpenCVDistortionModel
+{
+	struct FRadial
+	{
+		double K1 = 0.0;
+		double K2 = 0.0;
+		double K3 = 0.0;
+	};
+
+	struct FTangential
+	{
+		double P1 = 0.0;
+		double P2 = 0.0;
+	};
+
+	FRadial Radial;
+	FTangential Tangential;
+};
+
+struct FIphoneDistortionModel
+{
+	TArray<double> LensDistortionTable;
+	TArray<double> InverseLensDistortionTable;
+};
+
+class CAPTUREMANAGERMEDIARW_API FCoordinateSystem
+{
+public:
+
+	enum FDirection
+	{
+		Front = 0,
+		Back,
+		Right,
+		Left,
+		Up,
+		Down
+	};
+
+	FCoordinateSystem();
+	FCoordinateSystem(FMatrix InMatDescription);
+	FCoordinateSystem(FDirection InXDirection, FDirection InYDirection, FDirection InZDirection);
+
+	const FMatrix& GetMatDescription() const;
+
+private:
+
+	static int32 GetIndex(FDirection InDirection);
+	static int32 GetSign(FDirection InDirection);
+
+	FMatrix MatDescription;
+};
+
+inline bool operator==(const FCoordinateSystem& InLeft, const FCoordinateSystem& InRight)
+{
+	return InLeft.GetMatDescription() == InRight.GetMatDescription();
+}
+
+extern const FCoordinateSystem UnrealCS;
+extern const FCoordinateSystem OpenCvCS;
+
+struct FMediaCalibrationSample
+{
+	enum ECameraType
+	{
+		Video = 0,
+		Depth,
+
+		Unknown
+	};
+
+	FString CameraId;
+	ECameraType CameraType = ECameraType::Unknown;
+
+	FVector2D FocalLength;
+	FVector2D PrincipalPoint;
+	FTransform Transform;
+
+	FIntPoint Dimensions;
+	EMediaOrientation Orientation;
+
+	TVariant<FEmptyVariantState, FIphoneDistortionModel, FOpenCVDistortionModel> DistortionModel;
+	
+	FCoordinateSystem InputCoordinateSystem;
+};
+
+CAPTUREMANAGERMEDIARW_API FVector ConvertToCoordinateSystem(const FVector& InVector, 
+															const FCoordinateSystem& InInputCoordinateSystem, 
+															const FCoordinateSystem& InOutputCoordinateSystem);
+
+CAPTUREMANAGERMEDIARW_API FTransform ConvertToCoordinateSystem(const FTransform& InTransform,
+															   const FCoordinateSystem& InInputCoordinateSystem,
+															   const FCoordinateSystem& InOutputCoordinateSystem);
+
+
+}
diff --git a/Engine/Plugins/VirtualProduction/CaptureManager/CaptureManagerApp/Source/CaptureManagerPipeline/Public/CaptureManagerPipeline.h b/Engine/Plugins/VirtualProduction/CaptureManager/CaptureManagerApp/Source/CaptureManagerPipeline/Public/CaptureManagerPipeline.h
new file mode 100644
--- /dev/null
+++ b/Engine/Plugins/VirtualProduction/CaptureManager/CaptureManagerApp/Source/CaptureManagerPipeline/Public/CaptureManagerPipeline.h
@@ -0,0 +1,57 @@
+// Copyright Epic Games, Inc. All Rights Reserved.
+
+#pragma once
+
+#include "HAL/Platform.h"
+
+#include "Nodes/ConvertAudioNode.h"
+#include "Nodes/ConvertDepthNode.h"
+#include "Nodes/ConvertVideoNode.h"
+#include "Nodes/ConvertCalibrationNode.h"
+
+#include "Templates/SharedPointer.h"
+#include "Misc/Guid.h"
+
+#include "Async/Monitor.h"
+
+enum class EPipelineExecutionPolicy
+{
+	Asynchronous = 0,
+	Synchronous
+};
+
+class CAPTUREMANAGERPIPELINE_API FCaptureManagerPipeline : public TSharedFromThis<FCaptureManagerPipeline>
+{
+public:
+	using FResult = TMap<FGuid, FCaptureManagerPipelineNode::FResult>;
+
+	FCaptureManagerPipeline(EPipelineExecutionPolicy InExecutionPolicy);
+	~FCaptureManagerPipeline();
+
+	FGuid AddGenericNode(TSharedPtr<FCaptureManagerPipelineNode> InNode);
+	FGuid AddConvertVideoNode(TSharedPtr<FConvertVideoNode> InNode);
+	FGuid AddConvertAudioNode(TSharedPtr<FConvertAudioNode> InNode);
+	FGuid AddConvertDepthNode(TSharedPtr<FConvertDepthNode> InNode);
+	FGuid AddConvertCalibrationNode(TSharedPtr<FConvertCalibrationNode> InNode);
+
+	FGuid AddSyncedNode(TSharedPtr<FCaptureManagerPipelineNode> InNode);
+
+	// Blocking function
+	[[nodiscard]] FCaptureManagerPipeline::FResult Run();
+
+	void Cancel();
+
+private:
+
+	FGuid AddParallelPipelineNode(TSharedPtr<FCaptureManagerPipelineNode> InNode);
+
+	TPimplPtr<class FCaptureManagerPipelineImpl> Impl;
+
+	
+	using FNodeMap = TMap<FGuid, TSharedPtr<FCaptureManagerPipelineNode>>;
+
+	UE::CaptureManager::TMonitor<FNodeMap> ParallelNodes;
+	UE::CaptureManager::TMonitor<FNodeMap> SyncNodes;
+
+	EPipelineExecutionPolicy ExecutionPolicy;
+};
\ No newline at end of file
diff --git a/Engine/Plugins/VirtualProduction/CaptureManager/CaptureManagerApp/Source/CaptureManagerPipeline/Public/CaptureManagerPipelineNode.h b/Engine/Plugins/VirtualProduction/CaptureManager/CaptureManagerApp/Source/CaptureManagerPipeline/Public/CaptureManagerPipelineNode.h
new file mode 100644
--- /dev/null
+++ b/Engine/Plugins/VirtualProduction/CaptureManager/CaptureManagerApp/Source/CaptureManagerPipeline/Public/CaptureManagerPipelineNode.h
@@ -0,0 +1,50 @@
+// Copyright Epic Games, Inc. All Rights Reserved.
+
+#pragma once
+
+#include "HAL/Platform.h"
+#include "Internationalization/Text.h"
+#include "Templates/SharedPointer.h"
+#include "Templates/ValueOrError.h"
+
+#include "Templates/PimplPtr.h"
+
+class CAPTUREMANAGERPIPELINE_API FCaptureManagerPipelineError
+{
+public:
+	FCaptureManagerPipelineError(FText InMessage, int32 InCode = 0);
+
+	const FText& GetMessage() const;
+
+	int32 GetCode() const;
+
+private:
+
+	FText Message;
+	int32 Code;
+};
+
+class CAPTUREMANAGERPIPELINE_API FCaptureManagerPipelineNode : public TSharedFromThis<FCaptureManagerPipelineNode>
+{
+public:
+
+	using FResult = TValueOrError<void, FCaptureManagerPipelineError>;
+
+	FCaptureManagerPipelineNode(const FString& InName);
+	virtual ~FCaptureManagerPipelineNode();
+
+	FString GetName() const;
+
+	FResult Execute();
+	void Cancel();
+
+protected:
+
+	virtual FResult Prepare() = 0;
+	virtual FResult Run() = 0;
+	virtual FResult Validate() = 0;
+
+private:
+	class FCaptureManagerPipelineNodeImpl;
+	TPimplPtr<FCaptureManagerPipelineNodeImpl> Impl;
+};
\ No newline at end of file
diff --git a/Engine/Plugins/VirtualProduction/CaptureManager/CaptureManagerApp/Source/CaptureManagerPipeline/Public/Nodes/ConvertAudioNode.h b/Engine/Plugins/VirtualProduction/CaptureManager/CaptureManagerApp/Source/CaptureManagerPipeline/Public/Nodes/ConvertAudioNode.h
new file mode 100644
--- /dev/null
+++ b/Engine/Plugins/VirtualProduction/CaptureManager/CaptureManagerApp/Source/CaptureManagerPipeline/Public/Nodes/ConvertAudioNode.h
@@ -0,0 +1,31 @@
+// Copyright Epic Games, Inc. All Rights Reserved.
+
+#pragma once
+
+#include "CaptureManagerPipelineNode.h"
+
+#include "CaptureManagerTakeMetadata.h"
+
+class CAPTUREMANAGERPIPELINE_API FConvertAudioNode : public FCaptureManagerPipelineNode
+{
+public:
+
+	FConvertAudioNode(const FTakeMetadata::FAudio& InAudio,
+					  const FString& InOutputDirectory);
+
+	virtual ~FConvertAudioNode() override;
+
+protected:
+
+	FTakeMetadata::FAudio Audio;
+	FString OutputDirectory;
+
+private:
+
+	virtual FResult Prepare() override final;
+	virtual FResult Validate() override final;
+
+	FString GetAudioDirectory() const;
+
+	static FResult CheckForAudioFile(const FString& InAudioPath);
+};
\ No newline at end of file
diff --git a/Engine/Plugins/VirtualProduction/CaptureManager/CaptureManagerApp/Source/CaptureManagerPipeline/Public/Nodes/ConvertCalibrationNode.h b/Engine/Plugins/VirtualProduction/CaptureManager/CaptureManagerApp/Source/CaptureManagerPipeline/Public/Nodes/ConvertCalibrationNode.h
new file mode 100644
--- /dev/null
+++ b/Engine/Plugins/VirtualProduction/CaptureManager/CaptureManagerApp/Source/CaptureManagerPipeline/Public/Nodes/ConvertCalibrationNode.h
@@ -0,0 +1,26 @@
+// Copyright Epic Games, Inc. All Rights Reserved.
+
+#pragma once
+
+#include "CaptureManagerPipelineNode.h"
+
+#include "CaptureManagerTakeMetadata.h"
+
+class CAPTUREMANAGERPIPELINE_API FConvertCalibrationNode : public FCaptureManagerPipelineNode
+{
+public:
+
+	FConvertCalibrationNode(const FTakeMetadata::FCalibration& InCalibration,
+							const FString& InOutputDirectory);
+protected:
+
+	FTakeMetadata::FCalibration Calibration;
+	FString OutputDirectory;
+
+private:
+
+	virtual FResult Prepare() override final;
+	virtual FResult Validate() override final;
+
+	FString GetCalibrationDirectory() const;
+};
\ No newline at end of file
diff --git a/Engine/Plugins/VirtualProduction/CaptureManager/CaptureManagerApp/Source/CaptureManagerPipeline/Public/Nodes/ConvertDepthNode.h b/Engine/Plugins/VirtualProduction/CaptureManager/CaptureManagerApp/Source/CaptureManagerPipeline/Public/Nodes/ConvertDepthNode.h
new file mode 100644
--- /dev/null
+++ b/Engine/Plugins/VirtualProduction/CaptureManager/CaptureManagerApp/Source/CaptureManagerPipeline/Public/Nodes/ConvertDepthNode.h
@@ -0,0 +1,31 @@
+// Copyright Epic Games, Inc. All Rights Reserved.
+
+#pragma once
+
+#include "CaptureManagerPipelineNode.h"
+
+#include "CaptureManagerTakeMetadata.h"
+
+class CAPTUREMANAGERPIPELINE_API FConvertDepthNode : public FCaptureManagerPipelineNode
+{
+public:
+
+	FConvertDepthNode(const FTakeMetadata::FVideo& InDepth,
+					  const FString& InOutputDirectory);
+
+	virtual ~FConvertDepthNode() override;
+
+protected:
+
+	FTakeMetadata::FVideo Depth;
+	FString OutputDirectory;
+
+private:
+
+	virtual FResult Prepare() override final;
+	virtual FResult Validate() override final;
+
+	FString GetDepthDirectory() const;
+
+	static FResult CheckImagesForDepth(const FString& InDepthPath);
+};
\ No newline at end of file
diff --git a/Engine/Plugins/VirtualProduction/CaptureManager/CaptureManagerApp/Source/CaptureManagerPipeline/Public/Nodes/ConvertVideoNode.h b/Engine/Plugins/VirtualProduction/CaptureManager/CaptureManagerApp/Source/CaptureManagerPipeline/Public/Nodes/ConvertVideoNode.h
new file mode 100644
--- /dev/null
+++ b/Engine/Plugins/VirtualProduction/CaptureManager/CaptureManagerApp/Source/CaptureManagerPipeline/Public/Nodes/ConvertVideoNode.h
@@ -0,0 +1,31 @@
+// Copyright Epic Games, Inc. All Rights Reserved.
+
+#pragma once
+
+#include "CaptureManagerPipelineNode.h"
+
+#include "CaptureManagerTakeMetadata.h"
+
+class CAPTUREMANAGERPIPELINE_API FConvertVideoNode : public FCaptureManagerPipelineNode
+{
+public:
+
+	FConvertVideoNode(const FTakeMetadata::FVideo& InVideo, 
+					  const FString& InOutputDirectory);
+
+	virtual ~FConvertVideoNode() override;
+
+protected:
+
+	FTakeMetadata::FVideo Video;
+	FString OutputDirectory;
+
+private:
+
+	virtual FResult Prepare() override final;
+	virtual FResult Validate() override final;
+
+	FString GetVideoDirectory() const;
+
+	static FResult CheckImagesForVideo(const FString& InVideoPath);
+};
\ No newline at end of file
diff --git a/Engine/Plugins/VirtualProduction/CaptureManager/CaptureManagerApp/Source/CaptureManagerSettings/CaptureManagerSettings.Build.cs b/Engine/Plugins/VirtualProduction/CaptureManager/CaptureManagerApp/Source/CaptureManagerSettings/CaptureManagerSettings.Build.cs
--- a/Engine/Plugins/VirtualProduction/CaptureManager/CaptureManagerApp/Source/CaptureManagerSettings/CaptureManagerSettings.Build.cs
+++ b/Engine/Plugins/VirtualProduction/CaptureManager/CaptureManagerApp/Source/CaptureManagerSettings/CaptureManagerSettings.Build.cs
@@ -24,8 +24,7 @@ public class CaptureManagerSettings : ModuleRules
 			"SlateCore",
 			"Settings",
 			"UnrealEd",
-			"CaptureUtils",
-			"CaptureDataConverter"
+			"CaptureUtils"
 		});
 	}
 }
diff --git a/Engine/Plugins/VirtualProduction/CaptureManager/CaptureManagerApp/Source/CaptureManagerSettings/Private/Settings/CaptureManagerSettings.cpp b/Engine/Plugins/VirtualProduction/CaptureManager/CaptureManagerApp/Source/CaptureManagerSettings/Private/Settings/CaptureManagerSettings.cpp
--- a/Engine/Plugins/VirtualProduction/CaptureManager/CaptureManagerApp/Source/CaptureManagerSettings/Private/Settings/CaptureManagerSettings.cpp
+++ b/Engine/Plugins/VirtualProduction/CaptureManager/CaptureManagerApp/Source/CaptureManagerSettings/Private/Settings/CaptureManagerSettings.cpp
@@ -2,7 +2,6 @@
 
 #include "Settings/CaptureManagerSettings.h"
 #include "Settings/CaptureManagerTemplateTokens.h"
-#include "CaptureManagerEncoderNamingTokens.h"
 #include "Network/NetworkMisc.h"
 #include "Misc/Paths.h"
 
diff --git a/Engine/Plugins/VirtualProduction/CaptureManager/CaptureManagerApp/Source/CaptureManagerSettings/Private/Settings/CaptureManagerSettingsCustomization.cpp b/Engine/Plugins/VirtualProduction/CaptureManager/CaptureManagerApp/Source/CaptureManagerSettings/Private/Settings/CaptureManagerSettingsCustomization.cpp
--- a/Engine/Plugins/VirtualProduction/CaptureManager/CaptureManagerApp/Source/CaptureManagerSettings/Private/Settings/CaptureManagerSettingsCustomization.cpp
+++ b/Engine/Plugins/VirtualProduction/CaptureManager/CaptureManagerApp/Source/CaptureManagerSettings/Private/Settings/CaptureManagerSettingsCustomization.cpp
@@ -4,7 +4,6 @@
 
 #include "Settings/CaptureManagerTemplateTokens.h"
 #include "Settings/CaptureManagerSettings.h"
-#include "CaptureManagerEncoderNamingTokens.h"
 #include "DetailLayoutBuilder.h"
 #include "DetailWidgetRow.h"
 #include "Engine/Engine.h"
diff --git a/Engine/Plugins/VirtualProduction/CaptureManager/CaptureManagerApp/Source/CaptureManagerSettings/Public/Settings/CaptureManagerSettings.h b/Engine/Plugins/VirtualProduction/CaptureManager/CaptureManagerApp/Source/CaptureManagerSettings/Public/Settings/CaptureManagerSettings.h
--- a/Engine/Plugins/VirtualProduction/CaptureManager/CaptureManagerApp/Source/CaptureManagerSettings/Public/Settings/CaptureManagerSettings.h
+++ b/Engine/Plugins/VirtualProduction/CaptureManager/CaptureManagerApp/Source/CaptureManagerSettings/Public/Settings/CaptureManagerSettings.h
@@ -7,9 +7,6 @@
 #include "CaptureManagerSettings.generated.h"
 
 class UCaptureManagerGeneralTokens;
-
-// Encoder token classes live in CaptureDataConverter (Core) so they are auto-discovered
-// by UNamingTokensEngineSubsystem in both LiveLinkHub and UE Editor contexts.
 class UCaptureManagerVideoEncoderTokens;
 class UCaptureManagerAudioEncoderTokens;
 
diff --git a/Engine/Plugins/VirtualProduction/CaptureManager/CaptureManagerApp/Source/IngestLiveLinkDevice/IngestLiveLinkDevice.Build.cs b/Engine/Plugins/VirtualProduction/CaptureManager/CaptureManagerApp/Source/IngestLiveLinkDevice/IngestLiveLinkDevice.Build.cs
--- a/Engine/Plugins/VirtualProduction/CaptureManager/CaptureManagerApp/Source/IngestLiveLinkDevice/IngestLiveLinkDevice.Build.cs
+++ b/Engine/Plugins/VirtualProduction/CaptureManager/CaptureManagerApp/Source/IngestLiveLinkDevice/IngestLiveLinkDevice.Build.cs
@@ -40,8 +40,7 @@ public class IngestLiveLinkDevice : ModuleRules
 			"CoreUObject",
 			"UnrealEd",
 			"ElectraPlayerPlugin",
-			"MJPEGDecoderElectra",
-			"NamingTokens"
+			"MJPEGDecoderElectra"
 		});
 
 		CppCompileWarningSettings.UnsafeTypeCastWarningLevel = WarningLevel.Error;
diff --git a/Engine/Plugins/VirtualProduction/CaptureManager/CaptureManagerApp/Source/IngestLiveLinkDevice/Private/BaseIngestLiveLinkDevice.cpp b/Engine/Plugins/VirtualProduction/CaptureManager/CaptureManagerApp/Source/IngestLiveLinkDevice/Private/BaseIngestLiveLinkDevice.cpp
--- a/Engine/Plugins/VirtualProduction/CaptureManager/CaptureManagerApp/Source/IngestLiveLinkDevice/Private/BaseIngestLiveLinkDevice.cpp
+++ b/Engine/Plugins/VirtualProduction/CaptureManager/CaptureManagerApp/Source/IngestLiveLinkDevice/Private/BaseIngestLiveLinkDevice.cpp
@@ -5,76 +5,18 @@
 #include "HAL/FileManager.h"
 #include "CaptureManagerUnrealEndpointModule.h"
 #include "Settings/CaptureManagerSettings.h"
-#include "Settings/CaptureManagerTemplateTokens.h"
-#include "CaptureManagerEncoderNamingTokens.h"
 
 #include "Utils/IngestLiveLinkDeviceUtils.h"
 #include "Utils/IngestLiveLinkDeviceLog.h"
 
 #include "CaptureDataConverter.h"
-#include "Engine/Engine.h"
 
-#include "NamingTokensEngineSubsystem.h"
 #include "Async/TaskWaiter.h"
 
 DEFINE_LOG_CATEGORY(LogIngestLiveLinkDevice);
 
 namespace UE::CaptureManager::Private
 {
-
-static FString ResolveCommandArgs(const FStringView InDefaultArgs, const FString& InCustomArgs, const FString& InTokenNamespace)
-{
-	FString CommandArgs = InCustomArgs.IsEmpty() ? FString(InDefaultArgs) : InCustomArgs;
-
-	if (!GEngine)
-	{
-		UE_LOGF(LogIngestLiveLinkDevice, Warning, "Cannot resolve encoder command tokens: GEngine is unavailable");
-		return CommandArgs;
-	}
-
-	UNamingTokensEngineSubsystem* NamingTokensSubsystem = GEngine->GetEngineSubsystem<UNamingTokensEngineSubsystem>();
-
-	if (!NamingTokensSubsystem)
-	{
-		UE_LOGF(LogIngestLiveLinkDevice, Warning, "Cannot resolve encoder command tokens: NamingTokensSubsystem is unavailable");
-		return CommandArgs;
-	}
-
-	FNamingTokenFilterArgs FilterArgs;
-	FilterArgs.AdditionalNamespacesToInclude.Add(InTokenNamespace);
-	FilterArgs.bNativeOnly = true;
-
-	CommandArgs = NamingTokensSubsystem->EvaluateTokenString(CommandArgs, FilterArgs).EvaluatedText.ToString();
-
-	return CommandArgs;
-}
-
-static void ApplyThirdPartyEncoderSettings(const UCaptureManagerSettings* InSettings, FCaptureDataConverterParams& OutParams)
-{
-	if (!InSettings || !InSettings->bEnableThirdPartyEncoder || InSettings->ThirdPartyEncoder.FilePath.IsEmpty())
-	{
-		return;
-	}
-
-	const bool bThirdPartyEncoderIsAvailable = FPlatformProcess::ExecProcess(*InSettings->ThirdPartyEncoder.FilePath, TEXT("-version"), nullptr, nullptr, nullptr);
-
-	if (bThirdPartyEncoderIsAvailable)
-	{
-		// {input}, {output} and {params} tokens will remain and are resolved with an FString::Format in the encoder command's Execute() function
-		if (const UCaptureManagerAudioEncoderTokens* AudioTokens = InSettings->GetAudioEncoderNamingTokens())
-		{
-			const FString CommandArgs = ResolveCommandArgs(UE::CaptureManager::EncoderDefaults::AudioCommandArgs, InSettings->CustomAudioCommandArguments, AudioTokens->GetNamespace());
-			OutParams.AudioEncoderConfig = FCaptureManagerEncoderConfig{InSettings->ThirdPartyEncoder.FilePath, CommandArgs};
-		}
-
-		if (const UCaptureManagerVideoEncoderTokens* VideoTokens = InSettings->GetVideoEncoderNamingTokens())
-		{
-			const FString CommandArgs = ResolveCommandArgs(UE::CaptureManager::EncoderDefaults::VideoCommandArgs, InSettings->CustomVideoCommandArguments, VideoTokens->GetNamespace());
-			OutParams.VideoEncoderConfig = FCaptureManagerEncoderConfig{InSettings->ThirdPartyEncoder.FilePath, CommandArgs};
-		}
-	}
-}
-
 static void TryCleanUpWorkingDirectory(const FString& InWorkingDirectory)
 {
 	TRACE_CPUPROFILER_EVENT_SCOPE(FBaseIngestDevice_CleanUpWorkingDirectory);
@@ -467,11 +409,6 @@ FCaptureDataConverterResult<void> UBaseIngestLiveLinkDevice::FImpl::RunConversio
 
 	InParams.CalibrationOutputParams = MoveTemp(CalibrationParams);
 
-	if (const UCaptureManagerSettings* CaptureManagerSettings = GetDefault<UCaptureManagerSettings>())
-	{
-		Private::ApplyThirdPartyEncoderSettings(CaptureManagerSettings, InParams);
-	}
-
 	FCaptureDataConverter::FProgressReporter ProgressReporter =
 		FCaptureDataConverter::FProgressReporter::CreateLambda([ConvertTask](const double InProgress) mutable
 	{
diff --git a/Engine/Plugins/VirtualProduction/CaptureManager/CaptureManagerCore/CaptureManagerCore.uplugin b/Engine/Plugins/VirtualProduction/CaptureManager/CaptureManagerCore/CaptureManagerCore.uplugin
--- a/Engine/Plugins/VirtualProduction/CaptureManager/CaptureManagerCore/CaptureManagerCore.uplugin
+++ b/Engine/Plugins/VirtualProduction/CaptureManager/CaptureManagerCore/CaptureManagerCore.uplugin
@@ -27,21 +27,6 @@
 			"Type": "Runtime",
 			"LoadingPhase": "Default"
 		},
-		{
-			"Name": "CaptureManagerPipeline",
-			"Type": "Editor",
-			"LoadingPhase": "Default"
-		},
-		{
-			"Name": "CaptureManagerMediaRW",
-			"Type": "Editor",
-			"LoadingPhase": "Default"
-		},
-		{
-			"Name": "CaptureDataConverter",
-			"Type": "Editor",
-			"LoadingPhase": "Default"
-		},
 		{
 			"Name": "DataIngestCore",
 			"Type": "Editor",
@@ -67,18 +52,6 @@
 		{
 			"Name": "CaptureData",
 			"Enabled": true
-		},
-		{
-			"Name": "ControlFlows",
-			"Enabled": true
-		},
-		{
-			"Name": "NamingTokens",
-			"Enabled": true
-		},
-		{
-			"Name": "ElectraPlayer",
-			"Enabled": true
 		}
 	]
 }
\ No newline at end of file
