[CaptureManager] Relax iPhone take archive video naming restriction

When ingesting a legacy iPhone take archive (pre-UE 5.6, take.json format), the expected .mov filename is constructed deterministically from take metadata as {slate}_{take_number}_{subject}.mov. If the file on disk doesn't match this name, ingest would silently proceed with a bad path.

This change adds a fallback in ParseOldLiveLinkTakeMetadata: if the expected file is not found, the take directory is scanned for a single .mov file. If exactly one is found it is used in place of the expected path. If multiple are found, ingest bails with a validation error rather than guessing.

The correctly-named case and the live network ingest path (ULiveLinkFaceDevice) are unaffected, the fallback only runs when the expected file is absent.

#Synced-CL 52127547

diff --git a/Engine/Plugins/VirtualProduction/CaptureManager/CaptureManagerApp/Source/LiveLinkFaceMetadata/Private/LiveLinkFaceMetadata.cpp b/Engine/Plugins/VirtualProduction/CaptureManager/CaptureManagerApp/Source/LiveLinkFaceMetadata/Private/LiveLinkFaceMetadata.cpp
--- a/Engine/Plugins/VirtualProduction/CaptureManager/CaptureManagerApp/Source/LiveLinkFaceMetadata/Private/LiveLinkFaceMetadata.cpp
+++ b/Engine/Plugins/VirtualProduction/CaptureManager/CaptureManagerApp/Source/LiveLinkFaceMetadata/Private/LiveLinkFaceMetadata.cpp
@@ -130,6 +130,9 @@ struct FLiveLinkFaceTakeInfo
 	// The path to the folder where the files for this take are
 	FString TakeOriginDirectory;
 
+	// If set, overrides the name-derived video path (used when the .mov filename doesn't follow the expected convention)
+	TOptional<FString> VideoFilePathOverride;
+
 	TArray<uint8> RawThumbnailData;
 	FLiveLinkFaceTakeMetadata TakeMetadata;
 	FLiveLinkFaceVideoMetadata VideoMetadata;
@@ -276,6 +279,10 @@ FString FLiveLinkFaceTakeInfo::GetTakePath() const
 
 FString FLiveLinkFaceTakeInfo::GetVideoFilePath() const
 {
+	if (VideoFilePathOverride.IsSet())
+	{
+		return VideoFilePathOverride.GetValue();
+	}
 	return FString::Format(TEXT("{0}/{1}"), { TakeOriginDirectory, TakeMetadata.MOVFileName() });
 }
 
@@ -1030,6 +1037,40 @@ TOptional<FText> TakeDurationExceedsLimit(const float InDurationInSeconds)
 	return FText::Format(Message, FText::AsNumber(InDurationInSeconds, &Options), FText::AsNumber(Limit, &Options));
 }
 
+// Resolves the video file path for a take. If the expected file (named from the take metadata) is not found,
+// falls back to scanning the take directory for a single .mov file. Returns false if multiple .mov files are
+// found, as the correct one cannot be determined.
+static bool TryResolveVideoFilePath(const FString& InTakeDirectory, FLiveLinkFaceTakeInfo& OutTakeInfo)
+{
+	const FString ExpectedVideoPath = OutTakeInfo.GetVideoFilePath();
+	if (IFileManager::Get().FileExists(*ExpectedVideoPath))
+	{
+		return true;
+	}
+
+	TArray<FString> MovFiles;
+	IFileManager::Get().FindFiles(MovFiles, *(InTakeDirectory / (TEXT("*") + FLiveLinkFaceStaticFileNames::VideoExt)), /*bFiles=*/true, /*bDirectories=*/false);
+
+	if (MovFiles.Num() > 1)
+	{
+		UE_LOG(LogLiveLinkFaceMetadata, Warning,
+			TEXT("Multiple .mov files found in take directory '%s'; cannot determine correct video file for ingest"),
+			*InTakeDirectory);
+		return false;
+	}
+
+	if (MovFiles.Num() == 1)
+	{
+		UE_LOG(LogLiveLinkFaceMetadata, Log,
+			TEXT("Expected video file '%s' not found; using '%s' instead"),
+			*ExpectedVideoPath, *(InTakeDirectory / MovFiles[0]));
+		OutTakeInfo.VideoFilePathOverride = InTakeDirectory / MovFiles[0];
+	}
+
+	// If no .mov found at all, leave the path as-is. Downstream ingest will surface the missing file error.
+	return true;
+}
+
 #define REPORT_IF_ERROR_AND_RETURN(Result, Text) if (!Result) { OutValidationErrors.Add(Text); return {};}
 #define REPORT_IF_ERROR(Result, Text) if (!Result) { OutValidationErrors.Add(Text); }
 
@@ -1046,6 +1087,8 @@ TOptional<FTakeMetadata> ParseOldLiveLinkTakeMetadata(const FString& InTakeDirec
 	// Shouldn't block the ingest process if this fails
 	FLiveLinkFaceMetadataParser::ParseThumbnail(InTakeDirectory, TakeInfo);
 
+	REPORT_IF_ERROR_AND_RETURN(TryResolveVideoFilePath(InTakeDirectory, TakeInfo), LOCTEXT("ParseOldTakeMetadata_MultipleMovFiles", "Multiple .mov files found in take directory; cannot determine correct video file"));
+
 	REPORT_IF_ERROR(!ContainsWhitespace(TakeInfo.GetTakeName()), LOCTEXT("ParseOldTakeMetadata_TakeNameContainsWhiteSpace", "Take name contains white space character(s)"));
 	REPORT_IF_ERROR(!ContainsWhitespace(TakeInfo.TakeMetadata.Subject), LOCTEXT("ParseOldTakeMetadata_SubjectContainsWhiteSpace", "Subject contains white space character(s)"));
 	REPORT_IF_ERROR(!ContainsWhitespace(TakeInfo.TakeMetadata.SlateName), LOCTEXT("ParseOldTakeMetadata_SlateNameContainsWhiteSpace", "Slate name contains white space character(s)"));
