[Capture Manager] Improved handling of non-integer frame rates

#Synced-CL 52245287

diff --git a/Engine/Plugins/VirtualProduction/CaptureData/Source/CaptureDataUtils/Internal/SmpteTimecodeUtils.cpp b/Engine/Plugins/VirtualProduction/CaptureData/Source/CaptureDataUtils/Internal/SmpteTimecodeUtils.cpp
--- a/Engine/Plugins/VirtualProduction/CaptureData/Source/CaptureDataUtils/Internal/SmpteTimecodeUtils.cpp
+++ b/Engine/Plugins/VirtualProduction/CaptureData/Source/CaptureDataUtils/Internal/SmpteTimecodeUtils.cpp
@@ -30,6 +30,12 @@ FFrameRate EstimateSmpteTimecodeRate(const FFrameRate InMediaFrameRate)
 		return FFrameRate(24'000, 1'000);
 	}
 
+	if (FMath::IsNearlyEqual(InMediaFrameRate.AsDecimal(), 47.952, Tolerance))
+	{
+		// 23.976 = 24000/1001 (canonical NTSC pulldown rate)
+		return FFrameRate(24'000, 1'001);
+	}
+
 	return InMediaFrameRate;
 }
 
diff --git a/Engine/Plugins/VirtualProduction/CaptureManager/CaptureManagerApp/Source/IngestLiveLinkDevice/Private/Utils/CaptureExtractTimecode.cpp b/Engine/Plugins/VirtualProduction/CaptureManager/CaptureManagerApp/Source/IngestLiveLinkDevice/Private/Utils/CaptureExtractTimecode.cpp
--- a/Engine/Plugins/VirtualProduction/CaptureManager/CaptureManagerApp/Source/IngestLiveLinkDevice/Private/Utils/CaptureExtractTimecode.cpp
+++ b/Engine/Plugins/VirtualProduction/CaptureManager/CaptureManagerApp/Source/IngestLiveLinkDevice/Private/Utils/CaptureExtractTimecode.cpp
@@ -319,7 +319,7 @@ FCaptureExtractVideoInfo::FExtractResult FCaptureExtractVideoInfo::ExtractInfoUs
 {
 	using namespace UE::CaptureManager;
 
-	FString CommandArgs = FString::Format(TEXT("-v error -show_entries stream=codec_type,r_frame_rate,duration:format_tags=timecode:stream_side_data=rotation -of json \"{0}\""), { InFilePath });
+	FString CommandArgs = FString::Format(TEXT("-v error -show_entries stream=codec_type,r_frame_rate,duration:stream_tags=timecode:format_tags=timecode:stream_side_data=rotation -of json \"{0}\""), { InFilePath });
 
 	FProcessRunnerResult RunnerResult = FProcessRunner::Run(InFFProbePath, CommandArgs, nullptr, TimeoutPeriod);
 	
@@ -338,6 +338,8 @@ FCaptureExtractVideoInfo::FExtractResult FCaptureExtractVideoInfo::ExtractInfoUs
 		return MakeError(ECaptureExtractInfoError::FailedToParseJson);
 	}
 
+	bool bTimecodeFoundInStreamTags = false;
+
 	// Streams
 	const TArray<TSharedPtr<FJsonValue>>* StreamsArray;
 	if (JsonObject->TryGetArrayField(TEXT("streams"), StreamsArray))
@@ -380,6 +382,30 @@ FCaptureExtractVideoInfo::FExtractResult FCaptureExtractVideoInfo::ExtractInfoUs
 						}
 					}
 
+					// Stream-level timecode tag (e.g. MOV video track written by ffmpeg -metadata timecode=)
+					// Only use the first stream that has a parseable timecode tag.
+					if (!bTimecodeFoundInStreamTags)
+					{
+						const TSharedPtr<FJsonObject>* StreamTagsObject;
+						if (StreamObject->TryGetObjectField(TEXT("tags"), StreamTagsObject))
+						{
+							FString StreamTimecodeString;
+							if ((*StreamTagsObject)->TryGetStringField(TEXT("timecode"), StreamTimecodeString))
+							{
+								TOptional<FTimecode> TimecodeOpt = FTimecode::ParseTimecode(*StreamTimecodeString);
+								if (TimecodeOpt)
+								{
+									TimecodeInfo.Timecode = TimecodeOpt.GetValue();
+									bTimecodeFoundInStreamTags = true;
+								}
+								else
+								{
+									UE_LOGF(LogCaptureExtractTimecode, Verbose, "Failed to parse stream timecode string %ls, will attempt format-level fallback", *StreamTimecodeString);
+								}
+							}
+						}
+					}
+
 					// Video Orientation
 					int32 Rotation = 0;
 					const TArray<TSharedPtr<FJsonValue>>* SideDataArrayPointer;
@@ -414,25 +440,28 @@ FCaptureExtractVideoInfo::FExtractResult FCaptureExtractVideoInfo::ExtractInfoUs
 		}
 	}
 
-	// Timecode
-	const TSharedPtr<FJsonObject>* FormatObject;
-	if (JsonObject->TryGetObjectField(TEXT("format"), FormatObject)) 
+	// Timecode - fall back to format-level tags if not found in stream tags
+	if (!bTimecodeFoundInStreamTags)
 	{
-		const TSharedPtr<FJsonObject>* TagsObject;
-		if ((*FormatObject)->TryGetObjectField(TEXT("tags"), TagsObject)) 
+		const TSharedPtr<FJsonObject>* FormatObject;
+		if (JsonObject->TryGetObjectField(TEXT("format"), FormatObject)) 
 		{
-			FString TimecodeString;
-			if ((*TagsObject)->TryGetStringField(TEXT("timecode"), TimecodeString))
+			const TSharedPtr<FJsonObject>* TagsObject;
+			if ((*FormatObject)->TryGetObjectField(TEXT("tags"), TagsObject)) 
 			{
-				TOptional<FTimecode> TimecodeOpt = FTimecode::ParseTimecode(*TimecodeString);
-				if (TimecodeOpt)
-				{
-					TimecodeInfo.Timecode = TimecodeOpt.GetValue();
-				}
-				else
+				FString TimecodeString;
+				if ((*TagsObject)->TryGetStringField(TEXT("timecode"), TimecodeString))
 				{
-					UE_LOGF(LogCaptureExtractTimecode, Error, "Failed to parse timecode string %ls", *TimecodeString);
-					return MakeError(ECaptureExtractInfoError::UnableToParseTimecode);
+					TOptional<FTimecode> TimecodeOpt = FTimecode::ParseTimecode(*TimecodeString);
+					if (TimecodeOpt)
+					{
+						TimecodeInfo.Timecode = TimecodeOpt.GetValue();
+					}
+					else
+					{
+						UE_LOGF(LogCaptureExtractTimecode, Error, "Failed to parse timecode string %ls", *TimecodeString);
+						return MakeError(ECaptureExtractInfoError::UnableToParseTimecode);
+					}
 				}
 			}
 		}
@@ -460,7 +489,7 @@ TValueOrError<FFrameRate, UE::CaptureManager::ECaptureExtractInfoError> FCapture
 		FString Right;
 		TimecodeRateString.Split(" ", &Left, &Right);
 
-		double Number = FCString::Atoi(*Left);
+		double Number = FCString::Atof(*Left);
 		uint32 IntNumber;
 		uint32 Multiplier = 1;
 
