Matching device name with media source name

#Synced-CL 52680416

diff --git a/Engine/Plugins/Animation/LiveLinkDevice/Source/LiveLinkDevice/Private/LiveLinkDevice.cpp b/Engine/Plugins/Animation/LiveLinkDevice/Source/LiveLinkDevice/Private/LiveLinkDevice.cpp
--- a/Engine/Plugins/Animation/LiveLinkDevice/Source/LiveLinkDevice/Private/LiveLinkDevice.cpp
+++ b/Engine/Plugins/Animation/LiveLinkDevice/Source/LiveLinkDevice/Private/LiveLinkDevice.cpp
@@ -13,6 +13,17 @@
 #include "Widgets/SNullWidget.h"
 
 
+
+FText ULiveLinkDevice::GetDisplayName() const
+{
+	if (Settings)
+	{
+		return FText::FromString(Settings->DisplayName);
+	}
+	
+	return FText::GetEmpty();
+}
+
 TSharedRef<SWidget> ULiveLinkDevice::GenerateWidgetForColumn(const FName InColumnId, const FLiveLinkDeviceWidgetArguments& InArgs)
 {
 	// Handle generating a health widget for the status column.
diff --git a/Engine/Plugins/Animation/LiveLinkDevice/Source/LiveLinkDevice/Private/LiveLinkDeviceModule.cpp b/Engine/Plugins/Animation/LiveLinkDevice/Source/LiveLinkDevice/Private/LiveLinkDeviceModule.cpp
--- a/Engine/Plugins/Animation/LiveLinkDevice/Source/LiveLinkDevice/Private/LiveLinkDeviceModule.cpp
+++ b/Engine/Plugins/Animation/LiveLinkDevice/Source/LiveLinkDevice/Private/LiveLinkDeviceModule.cpp
@@ -2,6 +2,7 @@
 
 #include "LiveLinkDeviceModule.h"
 #include "DetailsViewArgs.h"
+#include "LiveLinkDeviceSettingsDetailCustomization.h"
 #include "Engine/Engine.h"
 #include "Framework/Docking/TabManager.h"
 #include "Framework/MultiBox/MultiBoxBuilder.h"
@@ -28,11 +29,26 @@ IMPLEMENT_MODULE(FLiveLinkDeviceModule, LiveLinkDevice);
 void FLiveLinkDeviceModule::StartupModule()
 {
 	FLiveLinkDeviceStyle::Initialize();
+
+	if (FModuleManager::Get().IsModuleLoaded("PropertyEditor"))
+	{
+		FPropertyEditorModule& PropertyModule = FModuleManager::GetModuleChecked<FPropertyEditorModule>("PropertyEditor");
+		PropertyModule.RegisterCustomClassLayout(
+			ULiveLinkDeviceSettings::StaticClass()->GetFName(),
+			FOnGetDetailCustomizationInstance::CreateStatic(&FLiveLinkDeviceSettingsDetailCustomization::MakeInstance)
+		);
+	}
 }
 
 
 void FLiveLinkDeviceModule::ShutdownModule()
 {
+	if (FModuleManager::Get().IsModuleLoaded("PropertyEditor"))
+	{
+		FPropertyEditorModule& PropertyModule = FModuleManager::GetModuleChecked<FPropertyEditorModule>("PropertyEditor");
+		PropertyModule.UnregisterCustomClassLayout(ULiveLinkDeviceSettings::StaticClass()->GetFName());
+	}
+
 	FLiveLinkDeviceStyle::Shutdown();
 }
 
diff --git a/Engine/Plugins/Animation/LiveLinkDevice/Source/LiveLinkDevice/Private/LiveLinkDeviceSettingsDetailCustomization.cpp b/Engine/Plugins/Animation/LiveLinkDevice/Source/LiveLinkDevice/Private/LiveLinkDeviceSettingsDetailCustomization.cpp
new file mode 100644
--- /dev/null
+++ b/Engine/Plugins/Animation/LiveLinkDevice/Source/LiveLinkDevice/Private/LiveLinkDeviceSettingsDetailCustomization.cpp
@@ -0,0 +1,123 @@
+// Copyright Epic Games, Inc. All Rights Reserved.
+
+#include "LiveLinkDeviceSettingsDetailCustomization.h"
+
+#include "DetailCategoryBuilder.h"
+#include "DetailLayoutBuilder.h"
+#include "DetailWidgetRow.h"
+#include "Engine/Engine.h"
+#include "IDetailPropertyRow.h"
+#include "LiveLinkDevice.h"
+#include "LiveLinkDeviceSubsystem.h"
+#include "Widgets/Input/SEditableTextBox.h"
+
+
+#define LOCTEXT_NAMESPACE "LiveLinkDevice"
+
+
+TSharedRef<IDetailCustomization> FLiveLinkDeviceSettingsDetailCustomization::MakeInstance()
+{
+	return MakeShareable(new FLiveLinkDeviceSettingsDetailCustomization());
+}
+
+void FLiveLinkDeviceSettingsDetailCustomization::CustomizeDetails(IDetailLayoutBuilder& DetailBuilder)
+{
+	DisplayNameHandle = DetailBuilder.GetProperty(GET_MEMBER_NAME_CHECKED(ULiveLinkDeviceSettings, DisplayName));
+
+	if (!DisplayNameHandle->IsValidHandle())
+	{
+		return;
+	}
+
+	// Resolve the owning device so we can exclude it from the uniqueness check.
+	TArray<TWeakObjectPtr<UObject>> ObjectsBeingCustomized;
+	DetailBuilder.GetObjectsBeingCustomized(ObjectsBeingCustomized);
+	if (ObjectsBeingCustomized.Num() > 0)
+	{
+		if (ULiveLinkDeviceSettings* Settings = Cast<ULiveLinkDeviceSettings>(ObjectsBeingCustomized[0].Get()))
+		{
+			WeakDevice = Settings->GetTypedOuter<ULiveLinkDevice>();
+		}
+	}
+
+	DetailBuilder.EditCategory(
+		DisplayNameHandle->GetDefaultCategoryName(),
+		FText::GetEmpty(),
+		ECategoryPriority::Important  // sorts above Default-priority categories
+	);
+
+	IDetailPropertyRow* Row = DetailBuilder.EditDefaultProperty(DisplayNameHandle.ToSharedRef());
+	Row->OverrideResetToDefault(FResetToDefaultOverride::Hide());
+
+	Row->CustomWidget()
+		.NameContent()
+		[
+			DisplayNameHandle->CreatePropertyNameWidget()
+		]
+		.ValueContent()
+		[
+			SNew(SEditableTextBox)
+			.Text(this, &FLiveLinkDeviceSettingsDetailCustomization::GetDisplayName)
+			.Font(IDetailLayoutBuilder::GetDetailFont())
+			.OnVerifyTextChanged(this, &FLiveLinkDeviceSettingsDetailCustomization::OnVerifyDisplayNameChanged)
+			.OnTextCommitted(this, &FLiveLinkDeviceSettingsDetailCustomization::SetDisplayName)
+			.SelectAllTextWhenFocused(true)
+		];
+}
+
+FText FLiveLinkDeviceSettingsDetailCustomization::GetDisplayName() const
+{
+	if (DisplayNameHandle.IsValid())
+	{
+		FText CurrentValue;
+		DisplayNameHandle->GetValueAsDisplayText(CurrentValue);
+		return CurrentValue;
+	}
+
+	return FText::GetEmpty();
+}
+
+bool FLiveLinkDeviceSettingsDetailCustomization::OnVerifyDisplayNameChanged(const FText& InText, FText& OutErrorMessage) const
+{
+	const FString TrimmedName = InText.ToString().TrimStartAndEnd();
+
+	if (TrimmedName.IsEmpty())
+	{
+		OutErrorMessage = LOCTEXT("LiveLinkDeviceSettings_EmptyName", "Display name cannot be empty.");
+		return false;
+	}
+
+	ULiveLinkDeviceSubsystem* Subsystem = GEngine ? GEngine->GetEngineSubsystem<ULiveLinkDeviceSubsystem>() : nullptr;
+	if (!Subsystem)
+	{
+		return true;
+	}
+
+	const ULiveLinkDevice* CurrentDevice = WeakDevice.Get();
+
+	for (const TPair<FGuid, TObjectPtr<ULiveLinkDevice>>& DevicePair : Subsystem->GetDeviceMap())
+	{
+		if (DevicePair.Value == CurrentDevice)
+		{
+			continue;
+		}
+
+		if (DevicePair.Value->GetDisplayName().ToString() == TrimmedName)
+		{
+			OutErrorMessage = LOCTEXT("LiveLinkDeviceSettings_DuplicateName", "A device with that name already exists.");
+			return false;
+		}
+	}
+
+	return true;
+}
+
+void FLiveLinkDeviceSettingsDetailCustomization::SetDisplayName(const FText& InCommittedText, ETextCommit::Type InCommitType)
+{
+	if (DisplayNameHandle.IsValid())
+	{
+		DisplayNameHandle->SetValue(InCommittedText.ToString().TrimStartAndEnd());
+	}
+}
+
+#undef LOCTEXT_NAMESPACE
diff --git a/Engine/Plugins/Animation/LiveLinkDevice/Source/LiveLinkDevice/Private/LiveLinkDeviceSettingsDetailCustomization.h b/Engine/Plugins/Animation/LiveLinkDevice/Source/LiveLinkDevice/Private/LiveLinkDeviceSettingsDetailCustomization.h
new file mode 100644
--- /dev/null
+++ b/Engine/Plugins/Animation/LiveLinkDevice/Source/LiveLinkDevice/Private/LiveLinkDeviceSettingsDetailCustomization.h
@@ -0,0 +1,35 @@
+// Copyright Epic Games, Inc. All Rights Reserved.
+
+#pragma once
+
+#include "IDetailCustomization.h"
+#include "UObject/WeakObjectPtr.h"
+#include "PropertyHandle.h"
+
+class ULiveLinkDevice;
+class ULiveLinkDeviceSettings;
+class ULiveLinkDeviceSubsystem;
+
+
+/**
+ * Detail customization for ULiveLinkDeviceSettings.
+ * Validates that DisplayName is unique across all devices.
+ */
+class FLiveLinkDeviceSettingsDetailCustomization : public IDetailCustomization
+{
+public:
+	static TSharedRef<IDetailCustomization> MakeInstance();
+
+	//~ Begin IDetailCustomization interface
+	virtual void CustomizeDetails(IDetailLayoutBuilder& DetailBuilder) override;
+	//~ End IDetailCustomization interface
+
+private:
+
+	FText GetDisplayName() const;
+	bool OnVerifyDisplayNameChanged(const FText& InText, FText& OutErrorMessage) const;
+	void SetDisplayName(const FText& InCommittedText, ETextCommit::Type InCommitType);
+
+	TWeakObjectPtr<ULiveLinkDevice> WeakDevice;
+	TSharedPtr<IPropertyHandle> DisplayNameHandle;
+};
diff --git a/Engine/Plugins/Animation/LiveLinkDevice/Source/LiveLinkDevice/Private/LiveLinkDeviceSubsystem.cpp b/Engine/Plugins/Animation/LiveLinkDevice/Source/LiveLinkDevice/Private/LiveLinkDeviceSubsystem.cpp
--- a/Engine/Plugins/Animation/LiveLinkDevice/Source/LiveLinkDevice/Private/LiveLinkDeviceSubsystem.cpp
+++ b/Engine/Plugins/Animation/LiveLinkDevice/Source/LiveLinkDevice/Private/LiveLinkDeviceSubsystem.cpp
@@ -215,6 +215,8 @@ ULiveLinkDeviceSubsystem::FCreateResult ULiveLinkDeviceSubsystem::InternalCreate
 		return MakeError();
 	}
 
+	NewSettings->DisplayName = GenerateUniqueDisplayName(NewSettings->DisplayName);
+
 	InternalAddDevice(DeviceGuid, NewDevice, NewSettings);
 
 	return MakeValue(DeviceGuid, NewDevice);
@@ -311,6 +313,37 @@ void ULiveLinkDeviceSubsystem::GetDevicesByCapability(TSubclassOf<ULiveLinkDevic
 	}
 }
 
+bool ULiveLinkDeviceSubsystem::IsDisplayNameInUse(const FString& InDisplayName) const
+{
+	for (const TPair<FGuid, TObjectPtr<ULiveLinkDevice>>& DevicePair : Devices)
+	{
+		if (DevicePair.Value->GetDisplayName().ToString() == InDisplayName)
+		{
+			return true;
+		}
+	}
+
+	return false;
+}
+
+FString ULiveLinkDeviceSubsystem::GenerateUniqueDisplayName(const FString& InBaseName) const
+{
+	if (!IsDisplayNameInUse(InBaseName))
+	{
+		return InBaseName;
+	}
+
+	int32 Suffix = 1;
+	FString Candidate = FString::Printf(TEXT("%s %d"), *InBaseName, Suffix);
+
+	while (IsDisplayNameInUse(Candidate))
+	{
+		++Suffix;
+		Candidate = FString::Printf(TEXT("%s %d"), *InBaseName, Suffix);
+	}
+
+	return Candidate;
+}
 
 void ULiveLinkDeviceSubsystem::RegisterCapabilityClass(TSubclassOf<ULiveLinkDeviceCapability> InCapability)
 {
diff --git a/Engine/Plugins/Animation/LiveLinkDevice/Source/LiveLinkDevice/Private/Tests/LiveLinkDevice_BasicTest.h b/Engine/Plugins/Animation/LiveLinkDevice/Source/LiveLinkDevice/Private/Tests/LiveLinkDevice_BasicTest.h
--- a/Engine/Plugins/Animation/LiveLinkDevice/Source/LiveLinkDevice/Private/Tests/LiveLinkDevice_BasicTest.h
+++ b/Engine/Plugins/Animation/LiveLinkDevice/Source/LiveLinkDevice/Private/Tests/LiveLinkDevice_BasicTest.h
@@ -33,6 +33,13 @@ UCLASS()
 class ULiveLinkDeviceSettings_BasicTest : public ULiveLinkDeviceSettings
 {
 	GENERATED_BODY()
+
+public:
+
+	ULiveLinkDeviceSettings_BasicTest()
+	{
+		DisplayName = TEXT("Test Device");
+	}
 };
 
 UCLASS(NotPlaceable)
@@ -44,7 +51,6 @@ class ULiveLinkDevice_BasicTest : public ULiveLinkDevice
 public:
 	//~ Begin ULiveLinkDevice interface
 	virtual TSubclassOf<ULiveLinkDeviceSettings> GetSettingsClass() const override { return ULiveLinkDeviceSettings_BasicTest::StaticClass(); }
-	virtual FText GetDisplayName() const override { return FText::FromString("Test Device"); }
 	virtual EDeviceHealth GetDeviceHealth() const override { return EDeviceHealth::Nominal; }
 	virtual FText GetHealthText() const override { return FText::FromString("Test Health"); }
 	//~ End ULiveLinkDevice interface
diff --git a/Engine/Plugins/Animation/LiveLinkDevice/Source/LiveLinkDevice/Public/LiveLinkDevice.h b/Engine/Plugins/Animation/LiveLinkDevice/Source/LiveLinkDevice/Public/LiveLinkDevice.h
--- a/Engine/Plugins/Animation/LiveLinkDevice/Source/LiveLinkDevice/Public/LiveLinkDevice.h
+++ b/Engine/Plugins/Animation/LiveLinkDevice/Source/LiveLinkDevice/Public/LiveLinkDevice.h
@@ -27,6 +27,12 @@ UCLASS(Abstract)
 class LIVELINKDEVICE_API ULiveLinkDeviceSettings : public UObject
 {
 	GENERATED_BODY()
+
+public:
+
+	/** Unique display name for this device */
+	UPROPERTY(EditAnywhere, BlueprintReadWrite, Category = "LiveLink Device")
+	FString DisplayName;
 };
 
 
@@ -44,7 +50,7 @@ public:
 
 	/** Human-readable display name for this device. */
 	UFUNCTION(BlueprintCallable, Category="Live Link Device")
-	virtual FText GetDisplayName() const PURE_VIRTUAL(ULiveLinkDevice::GetDisplayName, return FText(););
+	virtual FText GetDisplayName() const;
 
 	/** At-a-glance health/"severity" of the device. */
 	UFUNCTION(BlueprintCallable, Category="Live Link Device")
diff --git a/Engine/Plugins/Animation/LiveLinkDevice/Source/LiveLinkDevice/Public/LiveLinkDeviceSubsystem.h b/Engine/Plugins/Animation/LiveLinkDevice/Source/LiveLinkDevice/Public/LiveLinkDeviceSubsystem.h
--- a/Engine/Plugins/Animation/LiveLinkDevice/Source/LiveLinkDevice/Public/LiveLinkDeviceSubsystem.h
+++ b/Engine/Plugins/Animation/LiveLinkDevice/Source/LiveLinkDevice/Public/LiveLinkDeviceSubsystem.h
@@ -105,6 +105,12 @@ public:
 	UFUNCTION(BlueprintPure, Category="Live Link|Devices")
 	UE_API void GetDevicesByCapability(TSubclassOf<ULiveLinkDeviceCapability> Capability, TArray<ULiveLinkDevice*>& OutDevices) const;
 
+	/** Check whether a display name is already in use by any device. */
+	UE_API bool IsDisplayNameInUse(const FString& InDisplayName) const;
+	
+	/** Generate a unique display name by appending a numeric suffix if needed. */
+	UE_API FString GenerateUniqueDisplayName(const FString& InBaseName) const;
+
 private:
 	FCreateResult InternalCreateDeviceOfClass(
 		TSubclassOf<ULiveLinkDevice> InDeviceClass,
diff --git a/Engine/Plugins/Animation/LiveLinkExampleDevice/Source/LiveLinkExampleDevice/Private/Devices/LiveLinkExampleDevice.cpp b/Engine/Plugins/Animation/LiveLinkExampleDevice/Source/LiveLinkExampleDevice/Private/Devices/LiveLinkExampleDevice.cpp
--- a/Engine/Plugins/Animation/LiveLinkExampleDevice/Source/LiveLinkExampleDevice/Private/Devices/LiveLinkExampleDevice.cpp
+++ b/Engine/Plugins/Animation/LiveLinkExampleDevice/Source/LiveLinkExampleDevice/Private/Devices/LiveLinkExampleDevice.cpp
@@ -30,11 +30,6 @@ void ULiveLinkExampleDevice::OnDeviceRemoved()
 }
 
 
-FText ULiveLinkExampleDevice::GetDisplayName() const
-{
-	return FText::FromString(GetDeviceSettings<ULiveLinkExampleDeviceSettings>()->DisplayName);
-}
-
 EDeviceHealth ULiveLinkExampleDevice::GetDeviceHealth() const
 {
 	return EDeviceHealth::Nominal;
diff --git a/Engine/Plugins/Animation/LiveLinkExampleDevice/Source/LiveLinkExampleDevice/Private/Devices/LiveLinkExampleDevice.h b/Engine/Plugins/Animation/LiveLinkExampleDevice/Source/LiveLinkExampleDevice/Private/Devices/LiveLinkExampleDevice.h
--- a/Engine/Plugins/Animation/LiveLinkExampleDevice/Source/LiveLinkExampleDevice/Private/Devices/LiveLinkExampleDevice.h
+++ b/Engine/Plugins/Animation/LiveLinkExampleDevice/Source/LiveLinkExampleDevice/Private/Devices/LiveLinkExampleDevice.h
@@ -14,8 +14,11 @@ class ULiveLinkExampleDeviceSettings : public ULiveLinkDeviceSettings
     GENERATED_BODY()
 
 public:
-	UPROPERTY(EditAnywhere, Category="Example Device")
-	FString DisplayName = TEXT("Example Device");
+	
+	ULiveLinkExampleDeviceSettings()
+	{
+		DisplayName = TEXT("Example Device");
+	}
 
 	UPROPERTY(EditAnywhere, Category="Example Device")
 	FString IpAddress = TEXT("127.0.0.1");
@@ -39,7 +42,6 @@ public:
 		return ULiveLinkExampleDeviceSettings::StaticClass();
 	}
 
-	virtual FText GetDisplayName() const override;
 	virtual EDeviceHealth GetDeviceHealth() const override;
 	virtual FText GetHealthText() const override;
 
diff --git a/Engine/Plugins/Animation/LiveLinkGenericRecordingDevice/Source/LiveLinkGenericRecordingDevice/Private/LiveLinkGenericRecordingDevice.cpp b/Engine/Plugins/Animation/LiveLinkGenericRecordingDevice/Source/LiveLinkGenericRecordingDevice/Private/LiveLinkGenericRecordingDevice.cpp
--- a/Engine/Plugins/Animation/LiveLinkGenericRecordingDevice/Source/LiveLinkGenericRecordingDevice/Private/LiveLinkGenericRecordingDevice.cpp
+++ b/Engine/Plugins/Animation/LiveLinkGenericRecordingDevice/Source/LiveLinkGenericRecordingDevice/Private/LiveLinkGenericRecordingDevice.cpp
@@ -27,7 +27,7 @@ FText ULiveLinkGenericRecordingDevice::GetHealthText() const
 
 FText ULiveLinkGenericRecordingDevice::GetDisplayNameHelper_Implementation() const
 {
-	return FText::FromString(TEXT("Python Recording Device"));
+	return FText::FromString(GetDeviceSettings()->DisplayName);
 }
 
 EDeviceHealth ULiveLinkGenericRecordingDevice::GetDeviceHealthHelper_Implementation() const
diff --git a/Engine/Plugins/Animation/LiveLinkGenericRecordingDevice/Source/LiveLinkGenericRecordingDevice/Private/LiveLinkGenericRecordingDevice.h b/Engine/Plugins/Animation/LiveLinkGenericRecordingDevice/Source/LiveLinkGenericRecordingDevice/Private/LiveLinkGenericRecordingDevice.h
--- a/Engine/Plugins/Animation/LiveLinkGenericRecordingDevice/Source/LiveLinkGenericRecordingDevice/Private/LiveLinkGenericRecordingDevice.h
+++ b/Engine/Plugins/Animation/LiveLinkGenericRecordingDevice/Source/LiveLinkGenericRecordingDevice/Private/LiveLinkGenericRecordingDevice.h
@@ -11,6 +11,13 @@ UCLASS(Blueprintable)
 class ULiveLinkGenericRecordingDeviceSettings : public ULiveLinkDeviceSettings
 {
     GENERATED_BODY()
+
+public:
+
+	ULiveLinkGenericRecordingDeviceSettings()
+	{
+		DisplayName = TEXT("Python Recording Device");
+	}
 };
 
 UCLASS(Abstract, Blueprintable)
diff --git a/Engine/Plugins/Animation/LiveLinkUnrealDevice/Source/LiveLinkUnrealDevice/Private/Devices/LiveLinkUnrealDevice.cpp b/Engine/Plugins/Animation/LiveLinkUnrealDevice/Source/LiveLinkUnrealDevice/Private/Devices/LiveLinkUnrealDevice.cpp
--- a/Engine/Plugins/Animation/LiveLinkUnrealDevice/Source/LiveLinkUnrealDevice/Private/Devices/LiveLinkUnrealDevice.cpp
+++ b/Engine/Plugins/Animation/LiveLinkUnrealDevice/Source/LiveLinkUnrealDevice/Private/Devices/LiveLinkUnrealDevice.cpp
@@ -17,12 +17,6 @@ DEFINE_LOG_CATEGORY(LogLiveLinkUnrealDevice);
 #define LOCTEXT_NAMESPACE "LiveLinkUnrealDevice"
 
 
-FText ULiveLinkUnrealDevice::GetDisplayName() const
-{
-	return FText::FromString(GetSettings()->DisplayName);
-}
-
-
 EDeviceHealth ULiveLinkUnrealDevice::GetDeviceHealth() const
 {
 	return GetHealth().Key;
diff --git a/Engine/Plugins/Animation/LiveLinkUnrealDevice/Source/LiveLinkUnrealDevice/Private/Devices/LiveLinkUnrealDevice.h b/Engine/Plugins/Animation/LiveLinkUnrealDevice/Source/LiveLinkUnrealDevice/Private/Devices/LiveLinkUnrealDevice.h
--- a/Engine/Plugins/Animation/LiveLinkUnrealDevice/Source/LiveLinkUnrealDevice/Private/Devices/LiveLinkUnrealDevice.h
+++ b/Engine/Plugins/Animation/LiveLinkUnrealDevice/Source/LiveLinkUnrealDevice/Private/Devices/LiveLinkUnrealDevice.h
@@ -25,9 +25,11 @@ class ULiveLinkUnrealDeviceSettings : public ULiveLinkDeviceSettings
 	GENERATED_BODY()
 
 public:
-	/** User specified display name for the device. */
-	UPROPERTY(EditAnywhere, Category="Unreal Device")
-	FString DisplayName = NSLOCTEXT("LiveLinkUnrealDevice", "DefaultDisplayName", "Unreal Editor").ToString();
+
+	ULiveLinkUnrealDeviceSettings()
+	{
+		DisplayName = NSLOCTEXT("LiveLinkUnrealDevice", "DefaultDisplayName", "Unreal Editor").ToString();
+	}
 
 	/** The editor client known to Live Link Hub to connect to. */
 	UPROPERTY(EditAnywhere, Category="Unreal Device")
@@ -67,7 +69,6 @@ public:
 public:
 	virtual TSubclassOf<ULiveLinkDeviceSettings> GetSettingsClass() const override { return ULiveLinkUnrealDeviceSettings::StaticClass(); }
 
-	virtual FText GetDisplayName() const override;
 	virtual EDeviceHealth GetDeviceHealth() const override;
 	virtual FText GetHealthText() const override;
 
diff --git a/Engine/Plugins/VirtualProduction/CaptureManager/CaptureManagerApp/Source/ExampleLiveLinkDevices/Private/NetworkIngest/ExampleNetworkIngestDevice.h b/Engine/Plugins/VirtualProduction/CaptureManager/CaptureManagerApp/Source/ExampleLiveLinkDevices/Private/NetworkIngest/ExampleNetworkIngestDevice.h
--- a/Engine/Plugins/VirtualProduction/CaptureManager/CaptureManagerApp/Source/ExampleLiveLinkDevices/Private/NetworkIngest/ExampleNetworkIngestDevice.h
+++ b/Engine/Plugins/VirtualProduction/CaptureManager/CaptureManagerApp/Source/ExampleLiveLinkDevices/Private/NetworkIngest/ExampleNetworkIngestDevice.h
@@ -20,8 +20,11 @@ class UExampleNetworkIngestDeviceSettings : public ULiveLinkDeviceSettings
 	GENERATED_BODY()
 
 public:
-	UPROPERTY(EditAnywhere, Category = "Example Network Ingest Device")
-	FString DisplayName = TEXT("Example Network Ingest Device");
+
+	UExampleNetworkIngestDeviceSettings()
+	{
+		DisplayName = TEXT("Example Network Ingest Device");
+	}
 	
 	UPROPERTY(EditAnywhere, Category = "Example Network Ingest Device")
 	FString IpAddress;
diff --git a/Engine/Plugins/VirtualProduction/CaptureManager/CaptureManagerApp/Source/LiveLinkCapabilities/Private/Streaming/LiveLinkDeviceCapability_Streaming.cpp b/Engine/Plugins/VirtualProduction/CaptureManager/CaptureManagerApp/Source/LiveLinkCapabilities/Private/Streaming/LiveLinkDeviceCapability_Streaming.cpp
--- a/Engine/Plugins/VirtualProduction/CaptureManager/CaptureManagerApp/Source/LiveLinkCapabilities/Private/Streaming/LiveLinkDeviceCapability_Streaming.cpp
+++ b/Engine/Plugins/VirtualProduction/CaptureManager/CaptureManagerApp/Source/LiveLinkCapabilities/Private/Streaming/LiveLinkDeviceCapability_Streaming.cpp
@@ -18,6 +18,7 @@
 
 #define LOCTEXT_NAMESPACE "LiveLinkDeviceCapability_Streaming"
 
+
 void ULiveLinkDeviceCapability_Streaming::OnDeviceSubsystemInitialized()
 {
 	ULiveLinkDeviceSubsystem* Subsystem = GEngine->GetEngineSubsystem<ULiveLinkDeviceSubsystem>();
@@ -79,6 +80,7 @@ public:
 
 	UMediaSource* GetMediaSource() const;
 	void SetMediaSource(UMediaSource* InSource);
+	void SetMediaSourceName(const FString& InName);
 	void RegisterMediaSource(UMediaSource* InMediaSource);
 	void UnregisterMediaSource();
 	bool IsStarted() const { return bIsStarted; }
@@ -87,6 +89,7 @@ private:
 
 	void OnMediaProfileChanged(UMediaProfile* InPrevious, UMediaProfile* InNew);
 
+	void UpdateLabelForMediaSource_GameThread(const FString& InName, UMediaProfile* InMediaProfile);
 	void RegisterMediaSource_GameThread(UMediaSource* InMediaSource, UMediaProfile* InMediaProfile);
 	void UnregisterMediaSource_GameThread(UMediaProfile* InMediaProfile);
 
@@ -94,6 +97,9 @@ private:
 
 	mutable FCriticalSection Mutex;
 	TStrongObjectPtr<UMediaSource> CurrentMediaSource = nullptr;
+
+	// Updated only from GameThread
+	FString MediaSourceName;
 	std::atomic_bool bIsStarted = false;
 
 	FDelegateHandle DelegateHandle;
@@ -123,6 +129,40 @@ void ILiveLinkDeviceCapability_Streaming::FImpl::SetMediaSource(UMediaSource* In
 	CurrentMediaSource.Reset(InSource);
 }
 
+void ILiveLinkDeviceCapability_Streaming::FImpl::SetMediaSourceName(const FString& InName)
+{
+	if (IsInGameThread())
+	{
+		MediaSourceName = InName;
+
+		if (IsStarted())
+		{
+			IMediaProfileManager& MediaProfileManager = IMediaProfileManager::Get();
+			UMediaProfile* MediaProfile = MediaProfileManager.GetCurrentMediaProfile();
+
+			UpdateLabelForMediaSource_GameThread(MediaSourceName, MediaProfile);
+		}
+	}
+	else
+	{
+		ExecuteOnGameThread(TEXT("Updating Media Profile"), [This = AsWeak(), Name = InName]() mutable
+		{
+			if (TSharedPtr<FImpl> SharedThis = This.Pin())
+			{
+				SharedThis->MediaSourceName = MoveTemp(Name);
+
+				if (SharedThis->IsStarted())
+				{
+					IMediaProfileManager& MediaProfileManager = IMediaProfileManager::Get();
+					UMediaProfile* MediaProfile = MediaProfileManager.GetCurrentMediaProfile();
+
+					SharedThis->UpdateLabelForMediaSource_GameThread(SharedThis->MediaSourceName, MediaProfile);
+				}
+			}
+		});
+	}
+}
+
 void ILiveLinkDeviceCapability_Streaming::FImpl::RegisterMediaSource(UMediaSource* InMediaSource)
 {
 	TStrongObjectPtr<UMediaSource> MediaSource(InMediaSource);
@@ -133,6 +173,7 @@ void ILiveLinkDeviceCapability_Streaming::FImpl::RegisterMediaSource(UMediaSourc
 		UMediaProfile* MediaProfile = MediaProfileManager.GetCurrentMediaProfile();
 
 		RegisterMediaSource_GameThread(MediaSource.Get(), MediaProfile);
+		UpdateLabelForMediaSource_GameThread(MediaSourceName, MediaProfile);
 		bIsStarted = true;
 	}
 	else
@@ -145,6 +186,7 @@ void ILiveLinkDeviceCapability_Streaming::FImpl::RegisterMediaSource(UMediaSourc
 				UMediaProfile* MediaProfile = MediaProfileManager.GetCurrentMediaProfile();
 
 				SharedThis->RegisterMediaSource_GameThread(MediaSource.Get(), MediaProfile);
+				SharedThis->UpdateLabelForMediaSource_GameThread(SharedThis->MediaSourceName, MediaProfile);
 				SharedThis->bIsStarted = true;
 			}
 		});
@@ -199,9 +241,39 @@ void ILiveLinkDeviceCapability_Streaming::FImpl::OnMediaProfileChanged(UMediaPro
 	if (InNew)
 	{
 		RegisterMediaSource_GameThread(CurrentMediaSource.Get(), InNew);
+		UpdateLabelForMediaSource_GameThread(MediaSourceName, InNew);
 	}
 }
 
+void ILiveLinkDeviceCapability_Streaming::FImpl::UpdateLabelForMediaSource_GameThread(const FString& InName, UMediaProfile* InMediaProfile)
+{
+	if (!InMediaProfile)
+	{
+		return;
+	}
+
+#if WITH_EDITOR
+	int32 Index = InMediaProfile->FindMediaSourceIndex(GetMediaSource());
+	if (Index == INDEX_NONE)
+	{
+		return;
+	}
+
+	if (InMediaProfile->GetLabelForMediaSource(Index) == InName)
+	{
+		return;
+	}
+
+	FScopedTransaction UpdateMediaProfileTransaction(LOCTEXT("UpdateMediaProfileTransaction", "Update Media Profile"));
+	InMediaProfile->Modify();
+
+	// Will publish the property changed event
+	InMediaProfile->SetLabelForMediaSource(Index, InName);
+
+	InMediaProfile->MarkPackageDirty();
+#endif
+}
+
 void ILiveLinkDeviceCapability_Streaming::FImpl::RegisterMediaSource_GameThread(UMediaSource* InMediaSource, UMediaProfile* InMediaProfile)
 {
 	if (!InMediaProfile)
@@ -308,6 +380,11 @@ UMediaSource* ILiveLinkDeviceCapability_Streaming::GetMediaSource_Implementation
 	return Impl->GetMediaSource();
 }
 
+void ILiveLinkDeviceCapability_Streaming::SetMediaSourceName(const FString& InName)
+{
+	Impl->SetMediaSourceName(InName);
+}
+
 void ILiveLinkDeviceCapability_Streaming::SetMediaSource(UMediaSource* InMediaSource)
 {
 	Impl->SetMediaSource(InMediaSource);
diff --git a/Engine/Plugins/VirtualProduction/CaptureManager/CaptureManagerApp/Source/LiveLinkCapabilities/Public/Streaming/LiveLinkDeviceCapability_Streaming.h b/Engine/Plugins/VirtualProduction/CaptureManager/CaptureManagerApp/Source/LiveLinkCapabilities/Public/Streaming/LiveLinkDeviceCapability_Streaming.h
--- a/Engine/Plugins/VirtualProduction/CaptureManager/CaptureManagerApp/Source/LiveLinkCapabilities/Public/Streaming/LiveLinkDeviceCapability_Streaming.h
+++ b/Engine/Plugins/VirtualProduction/CaptureManager/CaptureManagerApp/Source/LiveLinkCapabilities/Public/Streaming/LiveLinkDeviceCapability_Streaming.h
@@ -60,6 +60,7 @@ public:
 
 protected:
 
+	void SetMediaSourceName(const FString& InName);
 	void SetMediaSource(class UMediaSource* InMediaSource);
 	void DeleteMediaSource();
 
diff --git a/Engine/Plugins/VirtualProduction/CaptureManager/CaptureManagerDevices/Source/CPSLiveLinkDevice/Private/LiveLinkFaceDevice.cpp b/Engine/Plugins/VirtualProduction/CaptureManager/CaptureManagerDevices/Source/CPSLiveLinkDevice/Private/LiveLinkFaceDevice.cpp
--- a/Engine/Plugins/VirtualProduction/CaptureManager/CaptureManagerDevices/Source/CPSLiveLinkDevice/Private/LiveLinkFaceDevice.cpp
+++ b/Engine/Plugins/VirtualProduction/CaptureManager/CaptureManagerDevices/Source/CPSLiveLinkDevice/Private/LiveLinkFaceDevice.cpp
@@ -30,11 +30,6 @@ TSubclassOf<ULiveLinkDeviceSettings> ULiveLinkFaceDevice::GetSettingsClass() con
 	return ULiveLinkFaceDeviceSettings::StaticClass();
 }
 
-FText ULiveLinkFaceDevice::GetDisplayName() const
-{
-	return FText::FromString(GetSettings()->DisplayName);
-}
-
 EDeviceHealth ULiveLinkFaceDevice::GetDeviceHealth() const
 {
 	return EDeviceHealth::Nominal;
@@ -42,7 +37,7 @@ EDeviceHealth ULiveLinkFaceDevice::GetDeviceHealth() const
 
 FText ULiveLinkFaceDevice::GetHealthText() const
 {
-	return FText::FromString("Example Health");
+	return FText::FromString(TEXT("Example Health"));
 }
 
 void ULiveLinkFaceDevice::OnDeviceAdded()
@@ -64,6 +59,16 @@ void ULiveLinkFaceDevice::OnDeviceRemoved()
 	Super::OnDeviceRemoved();
 }
 
+void ULiveLinkFaceDevice::OnSettingChanged(const FPropertyChangedEvent& InPropertyChangedEvent)
+{
+	if (InPropertyChangedEvent.GetPropertyName() == GET_MEMBER_NAME_CHECKED(ULiveLinkDeviceSettings, DisplayName))
+	{
+		SetMediaSourceName(GetDisplayName().ToString());
+	}
+
+	Super::OnSettingChanged(InPropertyChangedEvent);
+}
+
 TMap<TSubclassOf<ULiveLinkDeviceCapability>, TObjectPtr<UObject>> ULiveLinkFaceDevice::OnSaveDeviceCapabilityData() const
 {
 	TMap<TSubclassOf<ULiveLinkDeviceCapability>, TObjectPtr<UObject>> CapabilityData;
@@ -85,6 +90,7 @@ void ULiveLinkFaceDevice::OnLoadDeviceCapabilityData(TMap<TSubclassOf<ULiveLinkD
 			if (TObjectPtr<URtspMediaSource> MediaSource = Cast<URtspMediaSource>(CapabilityData.Value))
 			{
 				SetMediaSource(MediaSource);
+				SetMediaSourceName(GetDisplayName().ToString());
 				StartStreaming_Implementation();
 			}
 		}
@@ -361,6 +367,7 @@ void ULiveLinkFaceDevice::HandleConnectionChanged(TSharedPtr<const UE::CaptureMa
 			RtspMediaSource->bProvideCpuBuffer = true;
 
 			SetMediaSource(RtspMediaSource);
+			SetMediaSourceName(GetDisplayName().ToString());
 		}
 
 		// Start streaming as soon as device is connected
@@ -565,11 +572,11 @@ void ULiveLinkFaceDevice::FetchPreIngestFiles(TMap<FString, UE::CaptureManager::
 
 						FCPSDataStream::FData Data = Result.GetValue();
 
-						if (FileName == "thumbnail.jpg")
+						if (FileName == TEXT("thumbnail.jpg"))
 						{
 							TakeMetadata.Thumbnail = FTakeThumbnailData(MoveTemp(Data));
 						}
-						else if (FileName == "video_metadata.json")
+						else if (FileName == TEXT("video_metadata.json"))
 						{
 
 							FString DataString = FString(StringCast<UTF8CHAR>(reinterpret_cast<const char*>(Data.GetData()), Data.Num()));
@@ -590,7 +597,7 @@ void ULiveLinkFaceDevice::FetchPreIngestFiles(TMap<FString, UE::CaptureManager::
 
 	TUniquePtr<FCPSDataStream> DataStream = MakeUnique<FCPSDataStream>(MoveTemp(OnFileExportFinishedCallback));
 
-	Device->FetchFiles(MoveTemp(DataStream), { FString("thumbnail.jpg"), FString("video_metadata.json") });
+	Device->FetchFiles(MoveTemp(DataStream), { TEXT("thumbnail.jpg"), TEXT("video_metadata.json") });
 }
 
 void ULiveLinkFaceDevice::RemoveDownloadedTakeData(const UE::CaptureManager::FTakeId InTakeId)
diff --git a/Engine/Plugins/VirtualProduction/CaptureManager/CaptureManagerDevices/Source/CPSLiveLinkDevice/Public/LiveLinkFaceDevice.h b/Engine/Plugins/VirtualProduction/CaptureManager/CaptureManagerDevices/Source/CPSLiveLinkDevice/Public/LiveLinkFaceDevice.h
--- a/Engine/Plugins/VirtualProduction/CaptureManager/CaptureManagerDevices/Source/CPSLiveLinkDevice/Public/LiveLinkFaceDevice.h
+++ b/Engine/Plugins/VirtualProduction/CaptureManager/CaptureManagerDevices/Source/CPSLiveLinkDevice/Public/LiveLinkFaceDevice.h
@@ -27,8 +27,11 @@ class CPSLIVELINKDEVICE_API ULiveLinkFaceDeviceSettings : public ULiveLinkDevice
 	GENERATED_BODY()
 
 public:
-	UPROPERTY(EditAnywhere, BlueprintReadWrite, Category = "Live Link Face")
-	FString DisplayName = TEXT("Live Link Face");
+
+	ULiveLinkFaceDeviceSettings()
+	{
+		DisplayName = TEXT("Live Link Face");
+	}
 
 	UPROPERTY(EditAnywhere, BlueprintReadWrite, Category = "Live Link Face")
 	FDeviceIpAddress IpAddress;
@@ -55,12 +58,13 @@ public:
 
 	//~ Begin ULiveLinkDevice interface
 	virtual TSubclassOf<ULiveLinkDeviceSettings> GetSettingsClass() const override;
-	virtual FText GetDisplayName() const override;
 	virtual EDeviceHealth GetDeviceHealth() const override;
 	virtual FText GetHealthText() const override;
 
 	virtual void OnDeviceAdded() override;
 	virtual void OnDeviceRemoved() override;
+	virtual void OnSettingChanged(const FPropertyChangedEvent& InPropertyChangedEvent) override;
+
 	virtual TMap<TSubclassOf<ULiveLinkDeviceCapability>, TObjectPtr<UObject>> OnSaveDeviceCapabilityData() const override;
 	virtual void OnLoadDeviceCapabilityData(TMap<TSubclassOf<ULiveLinkDeviceCapability>, TObjectPtr<UObject>> InCapabilityDataMap) override;
 	//~ End ULiveLinkDevice interface
diff --git a/Engine/Plugins/VirtualProduction/CaptureManager/CaptureManagerDevices/Source/MonoVideoIngestDevice/Private/MonoVideoIngestDevice.cpp b/Engine/Plugins/VirtualProduction/CaptureManager/CaptureManagerDevices/Source/MonoVideoIngestDevice/Private/MonoVideoIngestDevice.cpp
--- a/Engine/Plugins/VirtualProduction/CaptureManager/CaptureManagerDevices/Source/MonoVideoIngestDevice/Private/MonoVideoIngestDevice.cpp
+++ b/Engine/Plugins/VirtualProduction/CaptureManager/CaptureManagerDevices/Source/MonoVideoIngestDevice/Private/MonoVideoIngestDevice.cpp
@@ -44,11 +44,6 @@ TSubclassOf<ULiveLinkDeviceSettings> UMonoVideoIngestDevice::GetSettingsClass()
 	return UMonoVideoIngestDeviceSettings::StaticClass();
 }
 
-FText UMonoVideoIngestDevice::GetDisplayName() const
-{
-	return FText::FromString(GetSettings()->DisplayName);
-}
-
 EDeviceHealth UMonoVideoIngestDevice::GetDeviceHealth() const
 {
 	return EDeviceHealth::Nominal;
diff --git a/Engine/Plugins/VirtualProduction/CaptureManager/CaptureManagerDevices/Source/MonoVideoIngestDevice/Private/MonoVideoIngestDevice.h b/Engine/Plugins/VirtualProduction/CaptureManager/CaptureManagerDevices/Source/MonoVideoIngestDevice/Private/MonoVideoIngestDevice.h
--- a/Engine/Plugins/VirtualProduction/CaptureManager/CaptureManagerDevices/Source/MonoVideoIngestDevice/Private/MonoVideoIngestDevice.h
+++ b/Engine/Plugins/VirtualProduction/CaptureManager/CaptureManagerDevices/Source/MonoVideoIngestDevice/Private/MonoVideoIngestDevice.h
@@ -26,11 +26,9 @@ public:
 	UMonoVideoIngestDeviceSettings()
 	{
 		VideoDiscoveryExpression = TEXT("<Auto>");
+		DisplayName = TEXT("Mono Video Ingest");
 	}
 
-	UPROPERTY(EditAnywhere, BlueprintReadWrite, Category = "Mono Video Ingest")
-	FString DisplayName = TEXT("Mono Video Ingest");
-
 	/* Path to a directory containing the take(s) data */
 	UPROPERTY(EditAnywhere, BlueprintReadWrite, Category = "Mono Video Ingest")
 	FDirectoryPath TakeDirectory;
@@ -74,7 +72,6 @@ public:
 
 	//~ Begin ULiveLinkDevice interface
 	virtual TSubclassOf<ULiveLinkDeviceSettings> GetSettingsClass() const override;
-	virtual FText GetDisplayName() const override;
 	virtual EDeviceHealth GetDeviceHealth() const override;
 	virtual FText GetHealthText() const override;
 
diff --git a/Engine/Plugins/VirtualProduction/CaptureManager/CaptureManagerDevices/Source/StereoVideoIngestDevice/Private/StereoVideoIngestDevice.cpp b/Engine/Plugins/VirtualProduction/CaptureManager/CaptureManagerDevices/Source/StereoVideoIngestDevice/Private/StereoVideoIngestDevice.cpp
--- a/Engine/Plugins/VirtualProduction/CaptureManager/CaptureManagerDevices/Source/StereoVideoIngestDevice/Private/StereoVideoIngestDevice.cpp
+++ b/Engine/Plugins/VirtualProduction/CaptureManager/CaptureManagerDevices/Source/StereoVideoIngestDevice/Private/StereoVideoIngestDevice.cpp
@@ -62,11 +62,6 @@ TSubclassOf<ULiveLinkDeviceSettings> UStereoVideoIngestDevice::GetSettingsClass(
 	return UStereoVideoIngestDeviceSettings::StaticClass();
 }
 
-FText UStereoVideoIngestDevice::GetDisplayName() const
-{
-	return FText::FromString(GetSettings()->DisplayName);
-}
-
 EDeviceHealth UStereoVideoIngestDevice::GetDeviceHealth() const
 {
 	return EDeviceHealth::Nominal;
diff --git a/Engine/Plugins/VirtualProduction/CaptureManager/CaptureManagerDevices/Source/StereoVideoIngestDevice/Private/StereoVideoIngestDevice.h b/Engine/Plugins/VirtualProduction/CaptureManager/CaptureManagerDevices/Source/StereoVideoIngestDevice/Private/StereoVideoIngestDevice.h
--- a/Engine/Plugins/VirtualProduction/CaptureManager/CaptureManagerDevices/Source/StereoVideoIngestDevice/Private/StereoVideoIngestDevice.h
+++ b/Engine/Plugins/VirtualProduction/CaptureManager/CaptureManagerDevices/Source/StereoVideoIngestDevice/Private/StereoVideoIngestDevice.h
@@ -27,10 +27,9 @@ public:
 	{
 		VideoDiscoveryExpression = TEXT("<Auto>");
 		AudioDiscoveryExpression = TEXT("<Auto>");
-	}
 
-	UPROPERTY(EditAnywhere, BlueprintReadWrite, Category = "Stereo Video Ingest")
-	FString DisplayName = TEXT("Stereo Video Ingest");
+		DisplayName = TEXT("Stereo Video Ingest");
+	}
 
 	/* Path to a directory containing the take(s) data */
 	UPROPERTY(EditAnywhere, BlueprintReadWrite, Category = "Stereo Video Ingest")
@@ -101,7 +100,6 @@ public:
 
 	//~ Begin ULiveLinkDevice interface
 	virtual TSubclassOf<ULiveLinkDeviceSettings> GetSettingsClass() const override;
-	virtual FText GetDisplayName() const override;
 	virtual EDeviceHealth GetDeviceHealth() const override;
 	virtual FText GetHealthText() const override;
 
diff --git a/Engine/Plugins/VirtualProduction/CaptureManager/CaptureManagerDevices/Source/TakeArchiveIngestDevice/Private/TakeArchiveIngestDevice.cpp b/Engine/Plugins/VirtualProduction/CaptureManager/CaptureManagerDevices/Source/TakeArchiveIngestDevice/Private/TakeArchiveIngestDevice.cpp
--- a/Engine/Plugins/VirtualProduction/CaptureManager/CaptureManagerDevices/Source/TakeArchiveIngestDevice/Private/TakeArchiveIngestDevice.cpp
+++ b/Engine/Plugins/VirtualProduction/CaptureManager/CaptureManagerDevices/Source/TakeArchiveIngestDevice/Private/TakeArchiveIngestDevice.cpp
@@ -24,11 +24,6 @@ TSubclassOf<ULiveLinkDeviceSettings> UTakeArchiveIngestDevice::GetSettingsClass(
 	return UTakeArchiveIngestDeviceSettings::StaticClass();
 }
 
-FText UTakeArchiveIngestDevice::GetDisplayName() const
-{
-	return FText::FromString(GetSettings()->DisplayName);
-}
-
 EDeviceHealth UTakeArchiveIngestDevice::GetDeviceHealth() const
 {
 	return EDeviceHealth::Nominal;
diff --git a/Engine/Plugins/VirtualProduction/CaptureManager/CaptureManagerDevices/Source/TakeArchiveIngestDevice/Private/TakeArchiveIngestDevice.h b/Engine/Plugins/VirtualProduction/CaptureManager/CaptureManagerDevices/Source/TakeArchiveIngestDevice/Private/TakeArchiveIngestDevice.h
--- a/Engine/Plugins/VirtualProduction/CaptureManager/CaptureManagerDevices/Source/TakeArchiveIngestDevice/Private/TakeArchiveIngestDevice.h
+++ b/Engine/Plugins/VirtualProduction/CaptureManager/CaptureManagerDevices/Source/TakeArchiveIngestDevice/Private/TakeArchiveIngestDevice.h
@@ -22,8 +22,11 @@ class UTakeArchiveIngestDeviceSettings : public ULiveLinkDeviceSettings
 	GENERATED_BODY()
 
 public:
-	UPROPERTY(EditAnywhere, BlueprintReadWrite, Category = "Take Archive Ingest")
-	FString DisplayName = TEXT("Take Archive Ingest");
+
+	UTakeArchiveIngestDeviceSettings()
+	{
+		DisplayName = TEXT("Take Archive Ingest");
+	}
 
 	/* Path to a directory containing the take(s) data */
 	UPROPERTY(EditAnywhere, BlueprintReadWrite, Category = "Take Archive Ingest")
@@ -43,7 +46,6 @@ public:
 
 	//~ Begin ULiveLinkDevice interface
 	virtual TSubclassOf<ULiveLinkDeviceSettings> GetSettingsClass() const override;
-	virtual FText GetDisplayName() const override;
 	virtual EDeviceHealth GetDeviceHealth() const override;
 	virtual FText GetHealthText() const override;
 
diff --git a/Engine/Plugins/VirtualProduction/LiveLinkKiProDevice/Source/LiveLinkKiProDevice/Private/LiveLinkKiProDeviceBase.cpp b/Engine/Plugins/VirtualProduction/LiveLinkKiProDevice/Source/LiveLinkKiProDevice/Private/LiveLinkKiProDeviceBase.cpp
--- a/Engine/Plugins/VirtualProduction/LiveLinkKiProDevice/Source/LiveLinkKiProDevice/Private/LiveLinkKiProDeviceBase.cpp
+++ b/Engine/Plugins/VirtualProduction/LiveLinkKiProDevice/Source/LiveLinkKiProDevice/Private/LiveLinkKiProDeviceBase.cpp
@@ -31,12 +31,6 @@ TSubclassOf<ULiveLinkDeviceSettings> ULiveLinkKiProDeviceBase::GetSettingsClass(
 	return SettingsClass ? SettingsClass : TSubclassOf<ULiveLinkDeviceSettings>(ULiveLinkKiProDeviceSettings::StaticClass());
 }
 
-FText ULiveLinkKiProDeviceBase::GetDisplayName() const
-{
-	const ULiveLinkKiProDeviceSettings* DeviceSettings = GetDeviceSettings<ULiveLinkKiProDeviceSettings>();
-	return FText::FromString(FString::Printf(TEXT("Ki Pro (%s)"), *DeviceSettings->IpAddress));
-}
-
 EDeviceHealth ULiveLinkKiProDeviceBase::GetDeviceHealth() const
 {
 	switch (ConnectionStatus)
@@ -83,6 +77,9 @@ void ULiveLinkKiProDeviceBase::OnDeviceAdded()
 			FTickerDelegate::CreateUObject(this, &ULiveLinkKiProDeviceBase::DeviceTick),
 			PollingIntervalSeconds
 		);
+
+	ULiveLinkKiProDeviceSettings* DeviceSettings = GetDeviceSettings<ULiveLinkKiProDeviceSettings>();
+	DeviceSettings->DisplayName = FString::Printf(TEXT("Ki Pro (%s)"), *DeviceSettings->IpAddress);
 }
 
 void ULiveLinkKiProDeviceBase::OnDeviceRemoved()
@@ -108,6 +105,11 @@ void ULiveLinkKiProDeviceBase::OnDeviceRemoved()
 
 void ULiveLinkKiProDeviceBase::OnSettingChanged(const FPropertyChangedEvent& InPropertyChangedEvent)
 {
+	if (InPropertyChangedEvent.GetPropertyName() == GET_MEMBER_NAME_CHECKED(ULiveLinkDeviceSettings, DisplayName))
+	{
+		return;
+	}
+
 	// When settings change, attempt immediate reconnection if disconnected
 	// This provides instant feedback when user updates IP address or port
 	if (ConnectionStatus == ELiveLinkDeviceConnectionStatus::Disconnected)
diff --git a/Engine/Plugins/VirtualProduction/LiveLinkKiProDevice/Source/LiveLinkKiProDevice/Public/LiveLinkKiProDeviceBase.h b/Engine/Plugins/VirtualProduction/LiveLinkKiProDevice/Source/LiveLinkKiProDevice/Public/LiveLinkKiProDeviceBase.h
--- a/Engine/Plugins/VirtualProduction/LiveLinkKiProDevice/Source/LiveLinkKiProDevice/Public/LiveLinkKiProDeviceBase.h
+++ b/Engine/Plugins/VirtualProduction/LiveLinkKiProDevice/Source/LiveLinkKiProDevice/Public/LiveLinkKiProDeviceBase.h
@@ -26,7 +26,6 @@ class LIVELINKKIPRODEVICE_API ULiveLinkKiProDeviceBase : public ULiveLinkDevice
 public:
 	//~ Begin ULiveLinkDevice interface
 	virtual TSubclassOf<ULiveLinkDeviceSettings> GetSettingsClass() const override;
-	virtual FText GetDisplayName() const override;
 	virtual EDeviceHealth GetDeviceHealth() const override;
 	virtual FText GetHealthText() const override;
 	virtual void OnDeviceAdded() override;
diff --git a/Engine/Plugins/VirtualProduction/LiveLinkOBSDevice/Source/LiveLinkOBSDevice/Private/Devices/LiveLinkOBSDevice.cpp b/Engine/Plugins/VirtualProduction/LiveLinkOBSDevice/Source/LiveLinkOBSDevice/Private/Devices/LiveLinkOBSDevice.cpp
--- a/Engine/Plugins/VirtualProduction/LiveLinkOBSDevice/Source/LiveLinkOBSDevice/Private/Devices/LiveLinkOBSDevice.cpp
+++ b/Engine/Plugins/VirtualProduction/LiveLinkOBSDevice/Source/LiveLinkOBSDevice/Private/Devices/LiveLinkOBSDevice.cpp
@@ -43,12 +43,6 @@ namespace OBSWebSocketOpCode
 // -----------------------------------------------------------------------
 // ULiveLinkDevice interface
 // -----------------------------------------------------------------------
-
-FText ULiveLinkOBSDevice::GetDisplayName() const
-{
-    return FText::FromString(GetSettings()->DisplayName);
-}
-
 EDeviceHealth ULiveLinkOBSDevice::GetDeviceHealth() const
 {
     if (bAuthFailed)
diff --git a/Engine/Plugins/VirtualProduction/LiveLinkOBSDevice/Source/LiveLinkOBSDevice/Private/Devices/LiveLinkOBSDevice.h b/Engine/Plugins/VirtualProduction/LiveLinkOBSDevice/Source/LiveLinkOBSDevice/Private/Devices/LiveLinkOBSDevice.h
--- a/Engine/Plugins/VirtualProduction/LiveLinkOBSDevice/Source/LiveLinkOBSDevice/Private/Devices/LiveLinkOBSDevice.h
+++ b/Engine/Plugins/VirtualProduction/LiveLinkOBSDevice/Source/LiveLinkOBSDevice/Private/Devices/LiveLinkOBSDevice.h
@@ -23,9 +23,11 @@ class ULiveLinkOBSDeviceSettings : public ULiveLinkDeviceSettings
     GENERATED_BODY()
 
 public:
-    /** User-visible name shown in the device table. */
-    UPROPERTY(EditAnywhere, Category="OBS Device")
-    FString DisplayName = TEXT("OBS Studio");
+
+	ULiveLinkOBSDeviceSettings()
+	{
+		DisplayName = TEXT("OBS Studio");
+	}
 
     /** Host running OBS WebSocket server (e.g. 127.0.0.1 or a remote IP). */
     UPROPERTY(EditAnywhere, Category="OBS Device")
@@ -74,7 +76,6 @@ public:
 
     //~ Begin ULiveLinkDevice interface
     virtual TSubclassOf<ULiveLinkDeviceSettings> GetSettingsClass() const override { return ULiveLinkOBSDeviceSettings::StaticClass(); }
-    virtual FText         GetDisplayName()  const override;
     virtual EDeviceHealth GetDeviceHealth() const override;
     virtual FText         GetHealthText()   const override;
     virtual void OnDeviceAdded()   override;
diff --git a/Engine/Plugins/VirtualProduction/LiveLinkSoundDevice/Source/LiveLinkSoundDevice/Private/LiveLinkSoundDeviceBase.cpp b/Engine/Plugins/VirtualProduction/LiveLinkSoundDevice/Source/LiveLinkSoundDevice/Private/LiveLinkSoundDeviceBase.cpp
--- a/Engine/Plugins/VirtualProduction/LiveLinkSoundDevice/Source/LiveLinkSoundDevice/Private/LiveLinkSoundDeviceBase.cpp
+++ b/Engine/Plugins/VirtualProduction/LiveLinkSoundDevice/Source/LiveLinkSoundDevice/Private/LiveLinkSoundDeviceBase.cpp
@@ -21,12 +21,6 @@ TSubclassOf<ULiveLinkDeviceSettings> ULiveLinkSoundDeviceBase::GetSettingsClass(
 	return SettingsClass ? SettingsClass : TSubclassOf<ULiveLinkDeviceSettings>(ULiveLinkSoundDeviceSettings::StaticClass());
 }
 
-FText ULiveLinkSoundDeviceBase::GetDisplayName() const
-{
-	const ULiveLinkSoundDeviceSettings* DeviceSettings = GetDeviceSettings<ULiveLinkSoundDeviceSettings>();
-	return FText::FromString(FString::Printf(TEXT("Sound Devices (%s)"), *DeviceSettings->IpAddress));
-}
-
 EDeviceHealth ULiveLinkSoundDeviceBase::GetDeviceHealth() const
 {
 	switch (ConnectionStatus)
@@ -71,6 +65,9 @@ void ULiveLinkSoundDeviceBase::OnDeviceAdded()
 		FTickerDelegate::CreateUObject(this, &ULiveLinkSoundDeviceBase::DeviceTick),
 		PollingIntervalSeconds
 	);
+
+	ULiveLinkSoundDeviceSettings* DeviceSettings = GetDeviceSettings<ULiveLinkSoundDeviceSettings>();
+	DeviceSettings->DisplayName = FString::Printf(TEXT("Sound Devices (%s)"), *DeviceSettings->IpAddress);
 }
 
 void ULiveLinkSoundDeviceBase::OnDeviceRemoved()
@@ -95,6 +92,11 @@ void ULiveLinkSoundDeviceBase::OnDeviceRemoved()
 
 void ULiveLinkSoundDeviceBase::OnSettingChanged(const FPropertyChangedEvent& InPropertyChangedEvent)
 {
+	if (InPropertyChangedEvent.GetPropertyName() == GET_MEMBER_NAME_CHECKED(ULiveLinkDeviceSettings, DisplayName))
+	{
+		return;
+	}
+
 	// When settings change, attempt immediate reconnection if disconnected
 	if (ConnectionStatus == ELiveLinkDeviceConnectionStatus::Disconnected)
 	{
diff --git a/Engine/Plugins/VirtualProduction/LiveLinkSoundDevice/Source/LiveLinkSoundDevice/Public/LiveLinkSoundDeviceBase.h b/Engine/Plugins/VirtualProduction/LiveLinkSoundDevice/Source/LiveLinkSoundDevice/Public/LiveLinkSoundDeviceBase.h
--- a/Engine/Plugins/VirtualProduction/LiveLinkSoundDevice/Source/LiveLinkSoundDevice/Public/LiveLinkSoundDeviceBase.h
+++ b/Engine/Plugins/VirtualProduction/LiveLinkSoundDevice/Source/LiveLinkSoundDevice/Public/LiveLinkSoundDeviceBase.h
@@ -27,7 +27,6 @@ class LIVELINKSOUNDDEVICE_API ULiveLinkSoundDeviceBase : public ULiveLinkDevice
 public:
 	//~ Begin ULiveLinkDevice interface
 	virtual TSubclassOf<ULiveLinkDeviceSettings> GetSettingsClass() const override;
-	virtual FText GetDisplayName() const override;
 	virtual EDeviceHealth GetDeviceHealth() const override;
 	virtual FText GetHealthText() const override;
 	virtual void OnDeviceAdded() override;
