[CaptureManager] Fix broken LLH encoder defaults.

#Synced-CL 52728346

diff --git a/Engine/Plugins/VirtualProduction/CaptureManager/CaptureManagerApp/Source/CaptureManagerSettings/Private/CaptureManagerSettingsModule.cpp b/Engine/Plugins/VirtualProduction/CaptureManager/CaptureManagerApp/Source/CaptureManagerSettings/Private/CaptureManagerSettingsModule.cpp
--- a/Engine/Plugins/VirtualProduction/CaptureManager/CaptureManagerApp/Source/CaptureManagerSettings/Private/CaptureManagerSettingsModule.cpp
+++ b/Engine/Plugins/VirtualProduction/CaptureManager/CaptureManagerApp/Source/CaptureManagerSettings/Private/CaptureManagerSettingsModule.cpp
@@ -8,9 +8,38 @@
 #include "PropertyEditorModule.h"
 #include "Settings/CaptureManagerSettings.h"
 #include "Settings/CaptureManagerSettingsCustomization.h"
+#include "CaptureManagerEncoderConfig.h"
 
 #define LOCTEXT_NAMESPACE "CaptureManagerSettings"
 
+namespace UE::CaptureManager::Private
+{
+
+// Populate empty encoder command fields with defaults and migrate pre-CL-40536694 placeholder names.
+static void InitializeEncoderCommandDefaults(UCaptureManagerSettings* InSettings)
+{
+	if (InSettings->CustomVideoCommandArguments.IsEmpty())
+	{
+		InSettings->CustomVideoCommandArguments = FString(EncoderDefaults::VideoCommandArgs);
+	}
+
+	if (InSettings->CustomAudioCommandArguments.IsEmpty())
+	{
+		InSettings->CustomAudioCommandArguments = FString(EncoderDefaults::AudioCommandArgs);
+	}
+
+	// CL 40536694 renamed placeholders from {InputFileName}/{OutputFileName}/{ConversionFilters}
+	// to {input}/{output}/{params}. Migrate any stale saved values.
+	InSettings->CustomVideoCommandArguments.ReplaceInline(TEXT("{InputFileName}"), TEXT("{input}"));
+	InSettings->CustomVideoCommandArguments.ReplaceInline(TEXT("{OutputFileName}"), TEXT("{output}"));
+	InSettings->CustomVideoCommandArguments.ReplaceInline(TEXT("{ConversionFilters}"), TEXT("{params}"));
+
+	InSettings->CustomAudioCommandArguments.ReplaceInline(TEXT("{InputFileName}"), TEXT("{input}"));
+	InSettings->CustomAudioCommandArguments.ReplaceInline(TEXT("{OutputFileName}"), TEXT("{output}"));
+}
+
+} // namespace UE::CaptureManager::Private
+
 void FCaptureManagerSettingsModule::StartupModule()
 {
 	FCoreDelegates::GetOnPostEngineInit().AddRaw(this, &FCaptureManagerSettingsModule::PostEngineInit);
@@ -18,6 +47,8 @@ void FCaptureManagerSettingsModule::StartupModule()
 
 	FPropertyEditorModule& PropertyModule = FModuleManager::LoadModuleChecked<FPropertyEditorModule>("PropertyEditor");
 	PropertyModule.RegisterCustomClassLayout(UCaptureManagerSettings::StaticClass()->GetFName(), FOnGetDetailCustomizationInstance::CreateStatic(&FCaptureManagerSettingsCustomization::MakeInstance));
+
+	UE::CaptureManager::Private::InitializeEncoderCommandDefaults(GetMutableDefault<UCaptureManagerSettings>());
 }
 
 void FCaptureManagerSettingsModule::ShutdownModule()
diff --git a/Engine/Plugins/VirtualProduction/CaptureManager/CaptureManagerCore/Source/CaptureDataConverter/Private/Nodes/ThirdPartyEncoder/CaptureConvertAudioDataThirdParty.cpp b/Engine/Plugins/VirtualProduction/CaptureManager/CaptureManagerCore/Source/CaptureDataConverter/Private/Nodes/ThirdPartyEncoder/CaptureConvertAudioDataThirdParty.cpp
--- a/Engine/Plugins/VirtualProduction/CaptureManager/CaptureManagerCore/Source/CaptureDataConverter/Private/Nodes/ThirdPartyEncoder/CaptureConvertAudioDataThirdParty.cpp
+++ b/Engine/Plugins/VirtualProduction/CaptureManager/CaptureManagerCore/Source/CaptureDataConverter/Private/Nodes/ThirdPartyEncoder/CaptureConvertAudioDataThirdParty.cpp
@@ -17,9 +17,6 @@
 #include "GlobalNamingTokens.h"
 #include "NamingTokenData.h"
 #include "NamingTokensEngineSubsystem.h"
-//#include "Settings/CaptureManagerSettings.h"
-//#include "Settings/CaptureManagerTemplateTokens.h"
-
 #include "ProcessRunner/ProcessRunner.h"
 
 #define LOCTEXT_NAMESPACE "CaptureConvertAudioDataTP"
@@ -103,8 +100,6 @@ FCaptureConvertAudioDataThirdParty::FResult FCaptureConvertAudioDataThirdParty::
 	const FString AudioOutputFile = FPaths::SetExtension(FPaths::Combine(DestinationDirectory, AudioParams.AudioFileName), AudioParams.Format);
 	FProcessRunnerResult RunnerResult = ThirdPartyCommand.Execute(AudioFilePath, AudioOutputFile);
 
-	//FProcessRunnerResult RunnerResult = FProcessRunner::Run(Command.ExecutablePath, Command.ExecutableArgs, &Params.StopToken);
-
 	if (RunnerResult.HasError())
 	{
 		EProcessRunnerError RunnerError = RunnerResult.StealError();
diff --git a/Engine/Plugins/VirtualProduction/CaptureManager/CaptureManagerCore/Source/CaptureDataConverter/Private/Nodes/ThirdPartyEncoder/CaptureConvertVideoDataThirdParty.cpp b/Engine/Plugins/VirtualProduction/CaptureManager/CaptureManagerCore/Source/CaptureDataConverter/Private/Nodes/ThirdPartyEncoder/CaptureConvertVideoDataThirdParty.cpp
--- a/Engine/Plugins/VirtualProduction/CaptureManager/CaptureManagerCore/Source/CaptureDataConverter/Private/Nodes/ThirdPartyEncoder/CaptureConvertVideoDataThirdParty.cpp
+++ b/Engine/Plugins/VirtualProduction/CaptureManager/CaptureManagerCore/Source/CaptureDataConverter/Private/Nodes/ThirdPartyEncoder/CaptureConvertVideoDataThirdParty.cpp
@@ -18,9 +18,6 @@
 #include "GlobalNamingTokens.h"
 #include "NamingTokenData.h"
 #include "NamingTokensEngineSubsystem.h"
-//#include "Settings/CaptureManagerSettings.h"
-//#include "Settings/CaptureManagerTemplateTokens.h"
-
 #include "IImageWrapperModule.h"
 
 #include "Internationalization/Regex.h"
@@ -123,7 +120,6 @@ FCaptureConvertVideoDataThirdParty::FCaptureConvertVideoDataThirdParty(UE::Captu
 																	   const FCaptureConvertDataNodeParams& InParams,
 																	   const FCaptureConvertVideoOutputParams& InVideoParams)
 	: FConvertVideoNode(InVideo, InOutputDirectory)
-	//, ThirdPartyEncoder(MoveTemp(InThirdPartyEncoder))
 	, ThirdPartyCommand(MoveTemp(InThirdPartyCommand))
 	, Params(InParams)
 	, VideoParams(InVideoParams)
@@ -210,8 +206,6 @@ FCaptureConvertVideoDataThirdParty::FResult FCaptureConvertVideoDataThirdParty::
 	const FString ConversionArguments = Private::CreateConversionArguments(VideoParams.OutputPixelFormat, VideoParams.Rotation);
 	FProcessRunnerResult RunnerResult = ThirdPartyCommand.Execute(VideoFilePath, ImageFilePath, ConversionArguments);
 
-	//FProcessRunnerResult RunnerResult = FProcessRunner::Run(Command.ExecutablePath, Command.ExecutableArgs, &Params.StopToken);
-
 	if (RunnerResult.HasError())
 	{
 		EProcessRunnerError RunnerError = RunnerResult.StealError();
diff --git a/Engine/Plugins/VirtualProduction/CaptureManager/CaptureManagerEditor/Source/CaptureManagerEditorSettings/Private/Settings/CaptureManagerEditorSettings.cpp b/Engine/Plugins/VirtualProduction/CaptureManager/CaptureManagerEditor/Source/CaptureManagerEditorSettings/Private/Settings/CaptureManagerEditorSettings.cpp
--- a/Engine/Plugins/VirtualProduction/CaptureManager/CaptureManagerEditor/Source/CaptureManagerEditorSettings/Private/Settings/CaptureManagerEditorSettings.cpp
+++ b/Engine/Plugins/VirtualProduction/CaptureManager/CaptureManagerEditor/Source/CaptureManagerEditorSettings/Private/Settings/CaptureManagerEditorSettings.cpp
@@ -2,6 +2,7 @@
 
 #include "Settings/CaptureManagerEditorSettings.h"
 #include "Settings/CaptureManagerEditorTemplateTokens.h"
+#include "CaptureManagerEncoderConfig.h"
 #include "CaptureManagerEncoderNamingTokens.h"
 #include "Misc/Paths.h"
 
@@ -446,6 +447,24 @@ void UCaptureManagerEditorSettings::InitializeValuesIfNotSet()
 	{
 		LensFileAssetName = Private::DefaultLensFileAssetName;
 	}
+
+	// Default empty encoder command fields and migrate legacy placeholders
+	if (CustomVideoCommandArguments.IsEmpty())
+	{
+		CustomVideoCommandArguments = FString(UE::CaptureManager::EncoderDefaults::VideoCommandArgs);
+	}
+
+	if (CustomAudioCommandArguments.IsEmpty())
+	{
+		CustomAudioCommandArguments = FString(UE::CaptureManager::EncoderDefaults::AudioCommandArgs);
+	}
+
+	CustomVideoCommandArguments.ReplaceInline(TEXT("{InputFileName}"), TEXT("{input}"));
+	CustomVideoCommandArguments.ReplaceInline(TEXT("{OutputFileName}"), TEXT("{output}"));
+	CustomVideoCommandArguments.ReplaceInline(TEXT("{ConversionFilters}"), TEXT("{params}"));
+
+	CustomAudioCommandArguments.ReplaceInline(TEXT("{InputFileName}"), TEXT("{input}"));
+	CustomAudioCommandArguments.ReplaceInline(TEXT("{OutputFileName}"), TEXT("{output}"));
 }
 
 void UCaptureManagerEditorSettings::OnHubConnectionEstablished(FGuid SourceId)
